Copying //UE4/Dev-Framework to //UE4/Dev-Main (Source: //UE4/Dev-Framework @ 3136612)

diff --git a/Engine/Source/Editor/AnimGraph/Private/EditModes/AnimNodeEditMode.cpp b/Engine/Source/Editor/AnimGraph/Private/EditModes/AnimNodeEditMode.cpp
new file mode 100644
--- /dev/null
+++ b/Engine/Source/Editor/AnimGraph/Private/EditModes/AnimNodeEditMode.cpp
@@ -0,0 +1,569 @@
+// Copyright 1998-2016 Epic Games, Inc. All Rights Reserved.
+
+#include "AnimGraphPrivatePCH.h"
+#include "AnimNodeEditMode.h"
+#include "IPersonaPreviewScene.h"
+#include "Animation/DebugSkelMeshComponent.h"
+#include "AnimNode_ModifyBone.h"
+#include "AnimGraphNode_SkeletalControlBase.h"
+#include "AssetEditorModeManager.h"
+
+#define LOCTEXT_NAMESPACE "AnimNodeEditMode"
+
+FAnimNodeEditMode::FAnimNodeEditMode()
+	: AnimNode(nullptr)
+	, RuntimeAnimNode(nullptr)
+	, bManipulating(false)
+	, bInTransaction(false)
+{
+	// Disable grid drawing for this mode as the viewport handles this
+	bDrawGrid = false;
+}
+
+bool FAnimNodeEditMode::GetCameraTarget(FSphere& OutTarget) const
+{
+	FVector Location(GetWidgetLocation());
+	OutTarget.Center = Location;
+	OutTarget.W = 50.0f;
+
+	return true;
+}
+
+IPersonaPreviewScene& FAnimNodeEditMode::GetAnimPreviewScene() const
+{
+	return *static_cast<IPersonaPreviewScene*>(static_cast<FAssetEditorModeManager*>(Owner)->GetPreviewScene());
+}
+
+void FAnimNodeEditMode::GetOnScreenDebugInfo(TArray<FText>& OutDebugInfo) const
+{
+	if (AnimNode != nullptr)
+	{
+		AnimNode->GetOnScreenDebugInfo(OutDebugInfo, GetAnimPreviewScene().GetPreviewMeshComponent());
+	}
+}
+
+ECoordSystem FAnimNodeEditMode::GetWidgetCoordinateSystem() const
+{
+	UAnimGraphNode_SkeletalControlBase* SkelControl = Cast<UAnimGraphNode_SkeletalControlBase>(AnimNode);
+	if (SkelControl != nullptr)
+	{
+		PRAGMA_DISABLE_DEPRECATION_WARNINGS
+		return (ECoordSystem)SkelControl->GetWidgetCoordinateSystem(GetAnimPreviewScene().GetPreviewMeshComponent());
+		PRAGMA_ENABLE_DEPRECATION_WARNINGS
+	}
+
+	return ECoordSystem::COORD_None;
+}
+
+FWidget::EWidgetMode FAnimNodeEditMode::GetWidgetMode() const
+{
+	UAnimGraphNode_SkeletalControlBase* SkelControl = Cast<UAnimGraphNode_SkeletalControlBase>(AnimNode);
+	if (SkelControl != nullptr)
+	{
+		PRAGMA_DISABLE_DEPRECATION_WARNINGS
+		return (FWidget::EWidgetMode)SkelControl->GetWidgetMode(GetAnimPreviewScene().GetPreviewMeshComponent());
+		PRAGMA_ENABLE_DEPRECATION_WARNINGS
+	}
+
+	return FWidget::EWidgetMode::WM_None;
+}
+
+FWidget::EWidgetMode FAnimNodeEditMode::ChangeToNextWidgetMode(FWidget::EWidgetMode CurWidgetMode)
+{
+	UAnimGraphNode_SkeletalControlBase* SkelControl = Cast<UAnimGraphNode_SkeletalControlBase>(AnimNode);
+	if (SkelControl != nullptr)
+	{
+		PRAGMA_DISABLE_DEPRECATION_WARNINGS
+		return (FWidget::EWidgetMode)SkelControl->ChangeToNextWidgetMode(GetAnimPreviewScene().GetPreviewMeshComponent(), CurWidgetMode);
+		PRAGMA_ENABLE_DEPRECATION_WARNINGS
+	}
+
+	return FWidget::EWidgetMode::WM_None;
+}
+
+bool FAnimNodeEditMode::SetWidgetMode(FWidget::EWidgetMode InWidgetMode)
+{
+	UAnimGraphNode_SkeletalControlBase* SkelControl = Cast<UAnimGraphNode_SkeletalControlBase>(AnimNode);
+	if (SkelControl != nullptr)
+	{
+		PRAGMA_DISABLE_DEPRECATION_WARNINGS
+		return SkelControl->SetWidgetMode(GetAnimPreviewScene().GetPreviewMeshComponent(), InWidgetMode);
+		PRAGMA_ENABLE_DEPRECATION_WARNINGS
+	}
+
+	return false;
+}
+
+FName FAnimNodeEditMode::GetSelectedBone() const
+{
+	UAnimGraphNode_SkeletalControlBase* SkelControl = Cast<UAnimGraphNode_SkeletalControlBase>(AnimNode);
+	if (SkelControl != nullptr)
+	{
+		PRAGMA_DISABLE_DEPRECATION_WARNINGS
+		return SkelControl->FindSelectedBone();
+		PRAGMA_ENABLE_DEPRECATION_WARNINGS
+	}
+
+	return NAME_None;
+}
+
+void FAnimNodeEditMode::EnterMode(UAnimGraphNode_Base* InEditorNode, FAnimNode_Base* InRuntimeNode)
+{
+	AnimNode = InEditorNode;
+	RuntimeAnimNode = InRuntimeNode;
+
+	UAnimGraphNode_SkeletalControlBase* SkelControl = Cast<UAnimGraphNode_SkeletalControlBase>(AnimNode);
+	if (SkelControl != nullptr)
+	{
+		PRAGMA_DISABLE_DEPRECATION_WARNINGS
+		SkelControl->MoveSelectActorLocation(GetAnimPreviewScene().GetPreviewMeshComponent(), (FAnimNode_SkeletalControlBase*)RuntimeAnimNode);
+		SkelControl->CopyNodeDataTo(RuntimeAnimNode);
+		PRAGMA_ENABLE_DEPRECATION_WARNINGS
+	}
+
+	GetModeManager()->SetCoordSystem(GetWidgetCoordinateSystem());
+	GetModeManager()->SetWidgetMode(GetWidgetMode());
+}
+
+void FAnimNodeEditMode::ExitMode()
+{
+	UAnimGraphNode_SkeletalControlBase* SkelControl = Cast<UAnimGraphNode_SkeletalControlBase>(AnimNode);
+	if (SkelControl != nullptr)
+	{
+		PRAGMA_DISABLE_DEPRECATION_WARNINGS
+		SkelControl->DeselectActor(GetAnimPreviewScene().GetPreviewMeshComponent());
+		PRAGMA_ENABLE_DEPRECATION_WARNINGS
+	}
+
+	AnimNode = nullptr;
+	RuntimeAnimNode = nullptr;
+}
+
+void FAnimNodeEditMode::Render(const FSceneView* View, FViewport* Viewport, FPrimitiveDrawInterface* PDI)
+{
+	if (AnimNode != nullptr)
+	{
+		AnimNode->Draw(PDI, GetAnimPreviewScene().GetPreviewMeshComponent());
+	}
+}
+
+void FAnimNodeEditMode::DrawHUD(FEditorViewportClient* ViewportClient, FViewport* Viewport, const FSceneView* View, FCanvas* Canvas)
+{
+	if (AnimNode != nullptr)
+	{
+		AnimNode->DrawCanvas(*Viewport, *const_cast<FSceneView*>(View), *Canvas, GetAnimPreviewScene().GetPreviewMeshComponent());
+	}
+}
+
+bool FAnimNodeEditMode::HandleClick(FEditorViewportClient* InViewportClient, HHitProxy* HitProxy, const FViewportClick& Click)
+{
+	if (HitProxy != nullptr && HitProxy->IsA(HActor::StaticGetType()))
+	{
+		HActor* ActorHitProxy = static_cast<HActor*>(HitProxy);
+		GetAnimPreviewScene().SetSelectedActor(ActorHitProxy->Actor);
+
+		UAnimGraphNode_SkeletalControlBase* SkelControl = Cast<UAnimGraphNode_SkeletalControlBase>(AnimNode);
+		if (SkelControl != nullptr)
+		{
+			PRAGMA_DISABLE_DEPRECATION_WARNINGS
+			SkelControl->ProcessActorClick(ActorHitProxy);
+			PRAGMA_ENABLE_DEPRECATION_WARNINGS
+		}
+		return true;
+	}
+
+	return false;
+}
+
+FVector FAnimNodeEditMode::GetWidgetLocation() const
+{
+	UAnimGraphNode_SkeletalControlBase* SkelControl = Cast<UAnimGraphNode_SkeletalControlBase>(AnimNode);
+	if (SkelControl != nullptr)
+	{
+		PRAGMA_DISABLE_DEPRECATION_WARNINGS
+		return SkelControl->GetWidgetLocation(GetAnimPreviewScene().GetPreviewMeshComponent(), (FAnimNode_SkeletalControlBase*)RuntimeAnimNode);
+		PRAGMA_ENABLE_DEPRECATION_WARNINGS
+	}
+
+	return FVector::ZeroVector;
+}
+
+bool FAnimNodeEditMode::StartTracking(FEditorViewportClient* InViewportClient, FViewport* InViewport)
+{
+	if (!bInTransaction)
+	{
+		GEditor->BeginTransaction(LOCTEXT("EditSkelControlNodeTransaction", "Edit Skeletal Control Node"));
+		AnimNode->SetFlags(RF_Transactional);
+		AnimNode->Modify();
+		bInTransaction = true;
+	}
+
+	bManipulating = true;
+
+	return true;
+}
+
+bool FAnimNodeEditMode::EndTracking(FEditorViewportClient* InViewportClient, FViewport* InViewport)
+{
+	if (bManipulating)
+	{
+		bManipulating = false;
+	}
+
+	if (bInTransaction)
+	{
+		GEditor->EndTransaction();
+		bInTransaction = false;
+	}
+
+	return true;
+}
+
+bool FAnimNodeEditMode::InputKey(FEditorViewportClient* InViewportClient, FViewport* InViewport, FKey InKey, EInputEvent InEvent)
+{
+	bool bHandled = false;
+
+	// Handle switching modes - only allowed when not already manipulating
+	if ((InEvent == IE_Pressed) && (InKey == EKeys::SpaceBar) && !bManipulating)
+	{
+		FWidget::EWidgetMode WidgetMode = (FWidget::EWidgetMode)ChangeToNextWidgetMode(GetModeManager()->GetWidgetMode());
+		GetModeManager()->SetWidgetMode(WidgetMode);
+		if (WidgetMode == FWidget::WM_Scale)
+		{
+			GetModeManager()->SetCoordSystem(COORD_Local);
+		}
+		else
+		{
+			GetModeManager()->SetCoordSystem(COORD_World);
+		}
+
+		bHandled = true;
+		InViewportClient->Invalidate();
+	}
+
+	return bHandled;
+}
+
+bool FAnimNodeEditMode::InputDelta(FEditorViewportClient* InViewportClient, FViewport* InViewport, FVector& InDrag, FRotator& InRot, FVector& InScale)
+{
+	const EAxisList::Type CurrentAxis = InViewportClient->GetCurrentWidgetAxis();
+	const FWidget::EWidgetMode WidgetMode = InViewportClient->GetWidgetMode();
+
+	bool bHandled = false;
+
+	UDebugSkelMeshComponent* PreviewMeshComponent = GetAnimPreviewScene().GetPreviewMeshComponent();
+
+	if (bManipulating && CurrentAxis != EAxisList::None)
+	{
+		bHandled = true;
+
+		const bool bDoRotation = WidgetMode == FWidget::WM_Rotate || WidgetMode == FWidget::WM_TranslateRotateZ;
+		const bool bDoTranslation = WidgetMode == FWidget::WM_Translate || WidgetMode == FWidget::WM_TranslateRotateZ;
+		const bool bDoScale = WidgetMode == FWidget::WM_Scale;
+
+		if (bDoRotation)
+		{
+			DoRotation(InRot);
+		}
+
+		if (bDoTranslation)
+		{
+			DoTranslation(InDrag);
+		}
+
+		if (bDoScale)
+		{
+			DoScale(InScale);
+		}
+
+		InViewport->Invalidate();
+	}
+
+	return bHandled;
+}
+
+bool FAnimNodeEditMode::GetCustomDrawingCoordinateSystem(FMatrix& InMatrix, void* InData)
+{
+	UDebugSkelMeshComponent* PreviewMeshComponent = GetAnimPreviewScene().GetPreviewMeshComponent();
+	FName BoneName = GetSelectedBone();
+	int32 BoneIndex = PreviewMeshComponent->GetBoneIndex(BoneName);
+	if (BoneIndex != INDEX_NONE)
+	{
+		FTransform BoneMatrix = PreviewMeshComponent->GetBoneTransform(BoneIndex);
+		InMatrix = BoneMatrix.ToMatrixNoScale().RemoveTranslation();
+		return true;
+	}
+
+	return false;
+}
+
+bool FAnimNodeEditMode::GetCustomInputCoordinateSystem(FMatrix& InMatrix, void* InData)
+{
+	return GetCustomDrawingCoordinateSystem(InMatrix, InData);
+}
+
+bool FAnimNodeEditMode::ShouldDrawWidget() const
+{
+	return true;
+}
+
+void FAnimNodeEditMode::DoTranslation(FVector& InTranslation)
+{
+	UAnimGraphNode_SkeletalControlBase* SkelControl = Cast<UAnimGraphNode_SkeletalControlBase>(AnimNode);
+	if (SkelControl != nullptr)
+	{
+		PRAGMA_DISABLE_DEPRECATION_WARNINGS
+		SkelControl->DoTranslation(GetAnimPreviewScene().GetPreviewMeshComponent(), InTranslation, (FAnimNode_SkeletalControlBase*)RuntimeAnimNode);
+		PRAGMA_ENABLE_DEPRECATION_WARNINGS
+	}
+}
+
+void FAnimNodeEditMode::DoRotation(FRotator& InRotation)
+{
+	UAnimGraphNode_SkeletalControlBase* SkelControl = Cast<UAnimGraphNode_SkeletalControlBase>(AnimNode);
+	if (SkelControl != nullptr)
+	{
+		PRAGMA_DISABLE_DEPRECATION_WARNINGS
+		SkelControl->DoRotation(GetAnimPreviewScene().GetPreviewMeshComponent(), InRotation, (FAnimNode_SkeletalControlBase*)RuntimeAnimNode);
+		PRAGMA_ENABLE_DEPRECATION_WARNINGS
+	}
+}
+
+void FAnimNodeEditMode::DoScale(FVector& InScale)
+{
+	UAnimGraphNode_SkeletalControlBase* SkelControl = Cast<UAnimGraphNode_SkeletalControlBase>(AnimNode);
+	if (SkelControl != nullptr)
+	{
+		PRAGMA_DISABLE_DEPRECATION_WARNINGS
+		SkelControl->DoScale(GetAnimPreviewScene().GetPreviewMeshComponent(), InScale, (FAnimNode_SkeletalControlBase*)RuntimeAnimNode);
+		PRAGMA_ENABLE_DEPRECATION_WARNINGS
+	}
+}
+
+void FAnimNodeEditMode::Tick(FEditorViewportClient* ViewportClient, float DeltaTime)
+{
+	IAnimNodeEditMode::Tick(ViewportClient, DeltaTime);
+
+	// Keep actor location in sync with animation
+	UAnimGraphNode_SkeletalControlBase* SkelControl = Cast<UAnimGraphNode_SkeletalControlBase>(AnimNode);
+	if (SkelControl != nullptr)
+	{
+		PRAGMA_DISABLE_DEPRECATION_WARNINGS
+		SkelControl->MoveSelectActorLocation(GetAnimPreviewScene().GetPreviewMeshComponent(), (FAnimNode_SkeletalControlBase*)RuntimeAnimNode);
+		PRAGMA_ENABLE_DEPRECATION_WARNINGS
+	}
+}
+
+void FAnimNodeEditMode::ConvertToComponentSpaceTransform(const USkeletalMeshComponent* SkelComp, const FTransform & InTransform, FTransform & OutCSTransform, int32 BoneIndex, EBoneControlSpace Space)
+{
+	USkeleton* Skeleton = SkelComp->SkeletalMesh->Skeleton;
+
+	switch (Space)
+	{
+	case BCS_WorldSpace:
+	{
+		OutCSTransform = InTransform;
+		OutCSTransform.SetToRelativeTransform(SkelComp->ComponentToWorld);
+	}
+	break;
+
+	case BCS_ComponentSpace:
+	{
+		// Component Space, no change.
+		OutCSTransform = InTransform;
+	}
+	break;
+
+	case BCS_ParentBoneSpace:
+		if (BoneIndex != INDEX_NONE)
+		{
+			const int32 ParentIndex = Skeleton->GetReferenceSkeleton().GetParentIndex(BoneIndex);
+			if (ParentIndex != INDEX_NONE)
+			{
+				const int32 MeshParentIndex = Skeleton->GetMeshBoneIndexFromSkeletonBoneIndex(SkelComp->SkeletalMesh, ParentIndex);
+				if (MeshParentIndex != INDEX_NONE)
+				{
+					const FTransform ParentTM = SkelComp->GetBoneTransform(MeshParentIndex);
+					OutCSTransform = InTransform * ParentTM;
+				}
+				else
+				{
+					OutCSTransform = InTransform;
+				}
+			}
+		}
+		break;
+
+	case BCS_BoneSpace:
+		if (BoneIndex != INDEX_NONE)
+		{
+			const int32 MeshBoneIndex = Skeleton->GetMeshBoneIndexFromSkeletonBoneIndex(SkelComp->SkeletalMesh, BoneIndex);
+			if (MeshBoneIndex != INDEX_NONE)
+			{
+				const FTransform BoneTM = SkelComp->GetBoneTransform(MeshBoneIndex);
+				OutCSTransform = InTransform * BoneTM;
+			}
+			else
+			{
+				OutCSTransform = InTransform;
+			}
+		}
+		break;
+
+	default:
+		if (SkelComp->SkeletalMesh)
+		{
+			UE_LOG(LogAnimation, Warning, TEXT("ConvertToComponentSpaceTransform: Unknown BoneSpace %d  for Mesh: %s"), (uint8)Space, *SkelComp->SkeletalMesh->GetFName().ToString());
+		}
+		else
+		{
+			UE_LOG(LogAnimation, Warning, TEXT("ConvertToComponentSpaceTransform: Unknown BoneSpace %d  for Skeleton: %s"), (uint8)Space, *Skeleton->GetFName().ToString());
+		}
+		break;
+	}
+}
+
+
+FVector FAnimNodeEditMode::ConvertCSVectorToBoneSpace(const USkeletalMeshComponent* SkelComp, FVector& InCSVector, FCSPose<FCompactHeapPose>& MeshBases, const FName& BoneName, const EBoneControlSpace Space)
+{
+	FVector OutVector = InCSVector;
+
+	if (MeshBases.GetPose().IsValid())
+	{
+		const FMeshPoseBoneIndex MeshBoneIndex(SkelComp->GetBoneIndex(BoneName));
+		const FCompactPoseBoneIndex BoneIndex = MeshBases.GetPose().GetBoneContainer().MakeCompactPoseIndex(MeshBoneIndex);
+
+		switch (Space)
+		{
+			// World Space, no change in preview window
+		case BCS_WorldSpace:
+		case BCS_ComponentSpace:
+			// Component Space, no change.
+			break;
+
+		case BCS_ParentBoneSpace:
+		{
+			const FCompactPoseBoneIndex ParentIndex = MeshBases.GetPose().GetParentBoneIndex(BoneIndex);
+			if (ParentIndex != INDEX_NONE)
+			{
+				const FTransform& ParentTM = MeshBases.GetComponentSpaceTransform(ParentIndex);
+				OutVector = ParentTM.InverseTransformVector(InCSVector);
+			}
+		}
+		break;
+
+		case BCS_BoneSpace:
+		{
+			const FTransform& BoneTM = MeshBases.GetComponentSpaceTransform(BoneIndex);
+			OutVector = BoneTM.InverseTransformVector(InCSVector);
+		}
+		break;
+		}
+	}
+
+	return OutVector;
+}
+
+FQuat FAnimNodeEditMode::ConvertCSRotationToBoneSpace(const USkeletalMeshComponent* SkelComp, FRotator& InCSRotator, FCSPose<FCompactHeapPose>& MeshBases, const FName& BoneName, const EBoneControlSpace Space)
+{
+	FQuat OutQuat = FQuat::Identity;
+
+	if (MeshBases.GetPose().IsValid())
+	{
+		const FMeshPoseBoneIndex MeshBoneIndex(SkelComp->GetBoneIndex(BoneName));
+		const FCompactPoseBoneIndex BoneIndex = MeshBases.GetPose().GetBoneContainer().MakeCompactPoseIndex(MeshBoneIndex);
+
+		FVector RotAxis;
+		float RotAngle;
+		InCSRotator.Quaternion().ToAxisAndAngle(RotAxis, RotAngle);
+
+		switch (Space)
+		{
+			// World Space, no change in preview window
+		case BCS_WorldSpace:
+		case BCS_ComponentSpace:
+			// Component Space, no change.
+			OutQuat = InCSRotator.Quaternion();
+			break;
+
+		case BCS_ParentBoneSpace:
+		{
+			const FCompactPoseBoneIndex ParentIndex = MeshBases.GetPose().GetParentBoneIndex(BoneIndex);
+			if (ParentIndex != INDEX_NONE)
+			{
+				const FTransform& ParentTM = MeshBases.GetComponentSpaceTransform(ParentIndex);
+				FTransform InverseParentTM = ParentTM.Inverse();
+				//Calculate the new delta rotation
+				FVector4 BoneSpaceAxis = InverseParentTM.TransformVector(RotAxis);
+				FQuat DeltaQuat(BoneSpaceAxis, RotAngle);
+				DeltaQuat.Normalize();
+				OutQuat = DeltaQuat;
+			}
+		}
+		break;
+
+		case BCS_BoneSpace:
+		{
+			const FTransform& BoneTM = MeshBases.GetComponentSpaceTransform(BoneIndex);
+			FTransform InverseBoneTM = BoneTM.Inverse();
+			FVector4 BoneSpaceAxis = InverseBoneTM.TransformVector(RotAxis);
+			//Calculate the new delta rotation
+			FQuat DeltaQuat(BoneSpaceAxis, RotAngle);
+			DeltaQuat.Normalize();
+			OutQuat = DeltaQuat;
+		}
+		break;
+		}
+	}
+
+	return OutQuat;
+}
+
+FVector FAnimNodeEditMode::ConvertWidgetLocation(const USkeletalMeshComponent* SkelComp, FCSPose<FCompactHeapPose>& MeshBases, const FName& BoneName, const FVector& Location, const EBoneControlSpace Space)
+{
+	FVector WidgetLoc = FVector::ZeroVector;
+
+	if (MeshBases.GetPose().IsValid())
+	{
+		USkeleton * Skeleton = SkelComp->SkeletalMesh->Skeleton;
+		const FMeshPoseBoneIndex MeshBoneIndex(SkelComp->GetBoneIndex(BoneName));
+		const FCompactPoseBoneIndex CompactBoneIndex = MeshBases.GetPose().GetBoneContainer().MakeCompactPoseIndex(MeshBoneIndex);
+
+		switch (Space)
+		{
+			// ComponentToWorld must be Identity in preview window so same as ComponentSpace
+		case BCS_WorldSpace:
+		case BCS_ComponentSpace:
+		{
+			// Component Space, no change.
+			WidgetLoc = Location;
+		}
+		break;
+
+		case BCS_ParentBoneSpace:
+
+			if (CompactBoneIndex != INDEX_NONE)
+			{
+				const FCompactPoseBoneIndex CompactParentIndex = MeshBases.GetPose().GetParentBoneIndex(CompactBoneIndex);
+				if (CompactParentIndex != INDEX_NONE)
+				{
+					const FTransform& ParentTM = MeshBases.GetComponentSpaceTransform(CompactParentIndex);
+					WidgetLoc = ParentTM.TransformPosition(Location);
+				}
+			}
+			break;
+
+		case BCS_BoneSpace:
+
+			if (CompactBoneIndex != INDEX_NONE)
+			{
+				const FTransform& BoneTM = MeshBases.GetComponentSpaceTransform(CompactBoneIndex);
+				WidgetLoc = BoneTM.TransformPosition(Location);
+			}
+		}
+	}
+
+	return WidgetLoc;
+}
+
+#undef LOCTEXT_NAMESPACE
\ No newline at end of file
diff --git a/Engine/Source/Editor/AnimGraph/Private/EditModes/FabrikEditMode.cpp b/Engine/Source/Editor/AnimGraph/Private/EditModes/FabrikEditMode.cpp
new file mode 100644
--- /dev/null
+++ b/Engine/Source/Editor/AnimGraph/Private/EditModes/FabrikEditMode.cpp
@@ -0,0 +1,60 @@
+// Copyright 1998-2016 Epic Games, Inc. All Rights Reserved.
+
+#include "AnimGraphPrivatePCH.h"
+#include "FabrikEditMode.h"
+#include "AnimGraphNode_Fabrik.h"
+#include "IPersonaPreviewScene.h"
+
+void FFabrikEditMode::EnterMode(class UAnimGraphNode_Base* InEditorNode, struct FAnimNode_Base* InRuntimeNode)
+{
+	RuntimeNode = static_cast<FAnimNode_Fabrik*>(InRuntimeNode);
+	GraphNode = CastChecked<UAnimGraphNode_Fabrik>(InEditorNode);
+
+	FAnimNodeEditMode::EnterMode(InEditorNode, InRuntimeNode);
+}
+
+void FFabrikEditMode::ExitMode()
+{
+	RuntimeNode = nullptr;
+	GraphNode = nullptr;
+
+	FAnimNodeEditMode::ExitMode();
+}
+
+FVector FFabrikEditMode::GetWidgetLocation() const
+{
+	USkeletalMeshComponent* SkelComp = GetAnimPreviewScene().GetPreviewMeshComponent();
+
+	FName& BoneName = GraphNode->Node.EffectorTransformBone.BoneName;
+	FVector Location = GraphNode->Node.EffectorTransform.GetLocation();
+	EBoneControlSpace Space = GraphNode->Node.EffectorTransformSpace;
+	FVector WidgetLoc = ConvertWidgetLocation(SkelComp, RuntimeNode->ForwardedPose, BoneName, Location, Space);
+	return WidgetLoc;
+}
+
+FWidget::EWidgetMode FFabrikEditMode::GetWidgetMode() const
+{
+	USkeletalMeshComponent* SkelComp = GetAnimPreviewScene().GetPreviewMeshComponent();
+	int32 BoneIndex = SkelComp->GetBoneIndex(GraphNode->Node.EffectorTransformBone.BoneName);
+
+	if (BoneIndex != INDEX_NONE)
+	{
+		return FWidget::WM_Translate;
+	}
+
+	return FWidget::WM_None;
+}
+
+FName FFabrikEditMode::GetSelectedBone() const
+{
+	return GraphNode->Node.EffectorTransformBone.BoneName;
+}
+
+void FFabrikEditMode::DoTranslation(FVector& InTranslation)
+{
+	USkeletalMeshComponent* SkelComp = GetAnimPreviewScene().GetPreviewMeshComponent();
+	FVector Offset = ConvertCSVectorToBoneSpace(SkelComp, InTranslation, RuntimeNode->ForwardedPose, GraphNode->Node.EffectorTransformBone.BoneName, GraphNode->Node.EffectorTransformSpace);
+
+	RuntimeNode->EffectorTransform.AddToTranslation(Offset);
+	GraphNode->Node.EffectorTransform.SetTranslation(RuntimeNode->EffectorTransform.GetTranslation());
+}
diff --git a/Engine/Source/Editor/AnimGraph/Private/EditModes/ModifyBoneEditMode.h b/Engine/Source/Editor/AnimGraph/Private/EditModes/ModifyBoneEditMode.h
new file mode 100644
--- /dev/null
+++ b/Engine/Source/Editor/AnimGraph/Private/EditModes/ModifyBoneEditMode.h
@@ -0,0 +1,38 @@
+// Copyright 1998-2016 Epic Games, Inc. All Rights Reserved.
+
+#pragma once
+
+#include "AnimNodeEditMode.h"
+#include "AnimNode_ModifyBone.h"
+
+class FModifyBoneEditMode : public FAnimNodeEditMode
+{
+public:
+	FModifyBoneEditMode();
+
+	/** IAnimNodeEditMode interface */
+	virtual void EnterMode(class UAnimGraphNode_Base* InEditorNode, struct FAnimNode_Base* InRuntimeNode) override;
+	virtual void ExitMode() override;
+	virtual ECoordSystem GetWidgetCoordinateSystem() const override;
+	virtual FVector GetWidgetLocation() const override;
+	virtual FWidget::EWidgetMode GetWidgetMode() const override;
+	virtual FWidget::EWidgetMode ChangeToNextWidgetMode(FWidget::EWidgetMode InCurWidgetMode) override;
+	virtual bool SetWidgetMode(FWidget::EWidgetMode InWidgetMode) override;
+	virtual FName GetSelectedBone() const override;
+	virtual void DoTranslation(FVector& InTranslation) override;
+	virtual void DoRotation(FRotator& InRotation) override;
+	virtual void DoScale(FVector& InScale) override;
+
+private:
+	// methods to find a valid widget mode for gizmo because doesn't need to show gizmo when the mode is "Ignore"
+	FWidget::EWidgetMode FindValidWidgetMode(FWidget::EWidgetMode InWidgetMode) const;
+	EBoneModificationMode GetBoneModificationMode(FWidget::EWidgetMode InWidgetMode) const;
+	FWidget::EWidgetMode GetNextWidgetMode(FWidget::EWidgetMode InWidgetMode) const;
+
+private:
+	struct FAnimNode_ModifyBone* RuntimeNode;
+	class UAnimGraphNode_ModifyBone* GraphNode;
+
+	// storing current widget mode 
+	mutable FWidget::EWidgetMode CurWidgetMode;
+};
\ No newline at end of file
diff --git a/Engine/Source/Editor/AnimGraph/Private/EditModes/PoseDriverEditMode.cpp b/Engine/Source/Editor/AnimGraph/Private/EditModes/PoseDriverEditMode.cpp
new file mode 100644
--- /dev/null
+++ b/Engine/Source/Editor/AnimGraph/Private/EditModes/PoseDriverEditMode.cpp
@@ -0,0 +1,130 @@
+// Copyright 1998-2016 Epic Games, Inc. All Rights Reserved.
+
+#include "AnimGraphPrivatePCH.h"
+#include "PoseDriverEditMode.h"
+#include "AnimGraphNode_PoseDriver.h"
+#include "IPersonaPreviewScene.h"
+
+#define LOCTEXT_NAMESPACE "A3Nodes"
+
+void FPoseDriverEditMode::EnterMode(class UAnimGraphNode_Base* InEditorNode, struct FAnimNode_Base* InRuntimeNode)
+{
+	RuntimeNode = static_cast<FAnimNode_PoseDriver*>(InRuntimeNode);
+	GraphNode = CastChecked<UAnimGraphNode_PoseDriver>(InEditorNode);
+
+	FAnimNodeEditMode::EnterMode(InEditorNode, InRuntimeNode);
+}
+
+void FPoseDriverEditMode::ExitMode()
+{
+	RuntimeNode = nullptr;
+	GraphNode = nullptr;
+
+	FAnimNodeEditMode::ExitMode();
+}
+
+static FLinearColor GetColorFromWeight(float InWeight)
+{
+	return FMath::Lerp(FLinearColor::White, FLinearColor::Red, InWeight);
+}
+
+void FPoseDriverEditMode::DrawHUD(FEditorViewportClient* ViewportClient, FViewport* Viewport, const FSceneView* View, FCanvas* Canvas)
+{
+	if (RuntimeNode == nullptr)
+	{
+		return;
+	}
+
+	// Get the node we are debugging
+	USkeletalMeshComponent* SkelComp = GetAnimPreviewScene().GetPreviewMeshComponent();
+
+	const int MarginX = 60;
+	const int XDelta = 10;
+
+	int PosY = 60;
+	const int YDelta = 20;
+
+	// Display bone we are watching
+	FText InterpItemText = FText::Format(LOCTEXT("BoneFormat", "Bone: {0}"), FText::FromName(RuntimeNode->SourceBone.BoneName));
+	FCanvasTextItem InterpItem(FVector2D(MarginX, PosY), InterpItemText, GEngine->GetSmallFont(), FLinearColor::White);
+	Canvas->DrawItem(InterpItem);
+	PosY += YDelta;
+
+	int32 PoseIndex = 0;
+	for (const FPoseDriverPoseInfo& PoseInfo : RuntimeNode->PoseInfos)
+	{
+		FString PoseItemString = FString::Printf(TEXT("%s : %f"), *PoseInfo.PoseName.ToString(), PoseInfo.PoseWeight);
+		FCanvasTextItem PoseItem(FVector2D(MarginX + XDelta, PosY), FText::FromString(PoseItemString), GEngine->GetSmallFont(), GetColorFromWeight(PoseInfo.PoseWeight));
+		Canvas->DrawItem(PoseItem);
+		PosY += YDelta;
+
+		PoseIndex++;
+	}
+}
+
+void FPoseDriverEditMode::Render(const FSceneView* View, FViewport* Viewport, FPrimitiveDrawInterface* PDI)
+{
+	USkeletalMeshComponent* SkelComp = GetAnimPreviewScene().GetPreviewMeshComponent();
+
+	static const float LineWidth = 0.1f;
+	static const float DrawAxisLength = 20.f;
+	static const float DrawPosSize = 1.f;
+
+	int32 BoneIndex = SkelComp->GetBoneIndex(RuntimeNode->SourceBone.BoneName);
+	if (BoneIndex != INDEX_NONE)
+	{
+		// Get transform of driven bone, used as basis for drawing
+		FTransform BoneWorldTM = SkelComp->GetBoneTransform(BoneIndex);
+		FVector BonePos = BoneWorldTM.GetLocation();
+
+		// Get ref frame for drawing - rotation from parent, translation from driver bone
+		FName ParentBoneName = SkelComp->GetParentBone(RuntimeNode->SourceBone.BoneName);
+		int32 ParentBoneIndex = SkelComp->GetBoneIndex(ParentBoneName);
+		FTransform ParentWorldTM = (ParentBoneIndex == INDEX_NONE) ? SkelComp->GetComponentToWorld() : SkelComp->GetBoneTransform(ParentBoneIndex);
+
+		// Swing drawing
+		if (RuntimeNode->Type == EPoseDriverType::SwingOnly)
+		{
+			FVector LocalVec = RuntimeNode->SourceBoneTM.TransformVectorNoScale(RuntimeNode->GetTwistAxisVector());
+			FVector WorldVec = ParentWorldTM.TransformVectorNoScale(LocalVec);
+			PDI->DrawLine(BonePos, BonePos + (WorldVec*DrawAxisLength), FLinearColor::Green, SDPG_Foreground, LineWidth);
+		}
+		// Swing & twist drawing
+		else if (RuntimeNode->Type == EPoseDriverType::SwingAndTwist)
+		{
+
+		}
+		// Translation drawing
+		else if (RuntimeNode->Type == EPoseDriverType::Translation)
+		{
+			FVector LocalPos = RuntimeNode->SourceBoneTM.GetTranslation();
+			FVector WorldPos = ParentWorldTM.TransformPosition(LocalPos);
+			DrawWireDiamond(PDI, FTranslationMatrix(WorldPos), DrawPosSize, FLinearColor::Green, SDPG_Foreground);
+		}
+
+		for (FPoseDriverPoseInfo& PoseInfo : RuntimeNode->PoseInfos)
+		{
+			// Swing drawing
+			if (RuntimeNode->Type == EPoseDriverType::SwingOnly)
+			{
+				FVector LocalVec = PoseInfo.PoseTM.TransformVectorNoScale(RuntimeNode->GetTwistAxisVector());
+				FVector WorldVec = ParentWorldTM.TransformVectorNoScale(LocalVec);
+				PDI->DrawLine(BonePos, BonePos + (WorldVec*DrawAxisLength), GetColorFromWeight(PoseInfo.PoseWeight), SDPG_Foreground, LineWidth);
+			}
+			// Swing & twist drawing
+			else if (RuntimeNode->Type == EPoseDriverType::SwingAndTwist)
+			{
+
+			}
+			// Translation drawing
+			else if (RuntimeNode->Type == EPoseDriverType::Translation)
+			{
+				FVector LocalPos = PoseInfo.PoseTM.GetTranslation();
+				FVector WorldPos = ParentWorldTM.TransformPosition(LocalPos);
+				DrawWireDiamond(PDI, FTranslationMatrix(WorldPos), DrawPosSize, GetColorFromWeight(PoseInfo.PoseWeight), SDPG_Foreground);
+			}
+		}
+	}
+}
+
+#undef LOCTEXT_NAMESPACE
diff --git a/Engine/Source/Editor/AnimGraph/Private/EditModes/TwoBoneIKEditMode.cpp b/Engine/Source/Editor/AnimGraph/Private/EditModes/TwoBoneIKEditMode.cpp
new file mode 100644
--- /dev/null
+++ b/Engine/Source/Editor/AnimGraph/Private/EditModes/TwoBoneIKEditMode.cpp
@@ -0,0 +1,263 @@
+// Copyright 1998-2016 Epic Games, Inc. All Rights Reserved.
+
+#include "AnimGraphPrivatePCH.h"
+#include "TwoBoneIKEditMode.h"
+#include "IPersonaPreviewScene.h"
+
+FTwoBoneIKEditMode::FTwoBoneIKEditMode()
+	: TwoBoneIKRuntimeNode(nullptr)
+	, TwoBoneIKGraphNode(nullptr)
+	, BoneSelectMode(BSM_EndEffector)
+{
+}
+
+void FTwoBoneIKEditMode::EnterMode(UAnimGraphNode_Base* InEditorNode, FAnimNode_Base* InRuntimeNode)
+{
+	TwoBoneIKRuntimeNode = static_cast<FAnimNode_TwoBoneIK*>(InRuntimeNode);
+	TwoBoneIKGraphNode = CastChecked<UAnimGraphNode_TwoBoneIK>(InEditorNode);
+
+	FAnimNodeEditMode::EnterMode(InEditorNode, InRuntimeNode);
+
+	ManageBoneSelectActor();
+}
+
+void FTwoBoneIKEditMode::ExitMode()
+{
+	// destroy bone select actor
+	if (BoneSelectActor != nullptr)
+	{
+		USkeletalMeshComponent* SkelComp = GetAnimPreviewScene().GetPreviewMeshComponent();
+		if (SkelComp->GetWorld()->DestroyActor(BoneSelectActor.Get()))
+		{
+			BoneSelectActor = nullptr;
+		}
+	}
+
+	TwoBoneIKGraphNode = nullptr;
+	TwoBoneIKRuntimeNode = nullptr;
+
+	FAnimNodeEditMode::ExitMode();
+}
+
+void FTwoBoneIKEditMode::Render(const FSceneView* View, FViewport* Viewport, FPrimitiveDrawInterface* PDI)
+{
+	USkeletalMeshComponent* SkelMeshComp = GetAnimPreviewScene().GetPreviewMeshComponent();
+	if (SkelMeshComp && SkelMeshComp->SkeletalMesh && SkelMeshComp->SkeletalMesh->Skeleton)
+	{
+		USkeleton* Skeleton = SkelMeshComp->SkeletalMesh->Skeleton;
+
+		DrawTargetLocation(PDI, SkelMeshComp, Skeleton, TwoBoneIKGraphNode->Node.EffectorLocationSpace, TwoBoneIKGraphNode->Node.EffectorSpaceBoneName, TwoBoneIKGraphNode->Node.EffectorLocation, FColor(255, 128, 128), FColor(180, 128, 128));
+		DrawTargetLocation(PDI, SkelMeshComp, Skeleton, TwoBoneIKGraphNode->Node.JointTargetLocationSpace, TwoBoneIKGraphNode->Node.JointTargetSpaceBoneName, TwoBoneIKGraphNode->Node.JointTargetLocation, FColor(128, 255, 128), FColor(128, 180, 128));
+	}
+}
+
+void FTwoBoneIKEditMode::DrawTargetLocation(FPrimitiveDrawInterface* PDI, USkeletalMeshComponent* SkelComp, USkeleton * Skeleton, EBoneControlSpace SpaceBase, FName SpaceBoneName, const FVector& TargetLocation, const FColor& TargetColor, const FColor& BoneColor) const
+{
+	const bool bInBoneSpace = (SpaceBase == BCS_ParentBoneSpace) || (SpaceBase == BCS_BoneSpace);
+	const int32 SpaceBoneIndex = bInBoneSpace ? Skeleton->GetReferenceSkeleton().FindBoneIndex(SpaceBoneName) : INDEX_NONE;
+	// Transform EffectorLocation from EffectorLocationSpace to ComponentSpace.
+	FTransform TargetTransform = FTransform (TargetLocation);
+	FTransform CSTransform;
+
+	ConvertToComponentSpaceTransform(SkelComp, TargetTransform, CSTransform, SpaceBoneIndex, SpaceBase);
+
+	FTransform WorldTransform = CSTransform * SkelComp->ComponentToWorld;
+
+	if (bInBoneSpace)
+	{
+		ConvertToComponentSpaceTransform(SkelComp, FTransform::Identity, CSTransform, SpaceBoneIndex, SpaceBase);
+		WorldTransform = CSTransform * SkelComp->ComponentToWorld;
+		DrawCoordinateSystem( PDI, WorldTransform.GetLocation(), WorldTransform.GetRotation().Rotator(), 20.f, SDPG_Foreground );
+		DrawWireDiamond( PDI, WorldTransform.ToMatrixWithScale(), 2.f, BoneColor, SDPG_Foreground );
+	}
+}
+
+FVector FTwoBoneIKEditMode::GetWidgetLocation() const
+{
+	EBoneControlSpace Space;
+	FName BoneName;
+	FVector Location;
+	 
+	if (BoneSelectMode == BSM_EndEffector)
+	{
+		Space = TwoBoneIKGraphNode->Node.EffectorLocationSpace;
+		Location = TwoBoneIKGraphNode->GetNodeValue(FString("EffectorLocation"), TwoBoneIKGraphNode->Node.EffectorLocation);
+		BoneName = TwoBoneIKGraphNode->Node.EffectorSpaceBoneName;
+
+	}
+	else // BSM_JointTarget
+	{
+		Space = TwoBoneIKGraphNode->Node.JointTargetLocationSpace;
+
+		if (Space == BCS_WorldSpace || Space == BCS_ComponentSpace)
+		{
+			return FVector::ZeroVector;
+		}
+		Location = TwoBoneIKGraphNode->GetNodeValue(FString("JointTargetLocation"), TwoBoneIKGraphNode->Node.JointTargetLocation);
+		BoneName = TwoBoneIKGraphNode->Node.JointTargetSpaceBoneName;
+	}
+
+	USkeletalMeshComponent* SkelComp = GetAnimPreviewScene().GetPreviewMeshComponent();
+	return ConvertWidgetLocation(SkelComp, TwoBoneIKRuntimeNode->ForwardedPose, BoneName, Location, Space);
+}
+
+FWidget::EWidgetMode FTwoBoneIKEditMode::GetWidgetMode() const
+{
+	int32 BoneIndex = GetAnimPreviewScene().GetPreviewMeshComponent()->GetBoneIndex(TwoBoneIKGraphNode->Node.IKBone.BoneName);
+	// Two bone IK node uses only Translate
+	if (BoneIndex != INDEX_NONE)
+	{
+		return FWidget::WM_Translate;
+	}
+
+	return FWidget::WM_None;
+}
+
+FName FTwoBoneIKEditMode::GetSelectedBone() const
+{
+	FName SelectedBone;
+
+	// should return mesh bone index
+	if (BoneSelectMode == BSM_EndEffector)
+	{
+		if (TwoBoneIKGraphNode->Node.EffectorLocationSpace == EBoneControlSpace::BCS_BoneSpace ||
+			TwoBoneIKGraphNode->Node.EffectorLocationSpace == EBoneControlSpace::BCS_ParentBoneSpace)
+		{
+			SelectedBone = TwoBoneIKGraphNode->Node.EffectorSpaceBoneName;
+		}
+		else
+		{
+			SelectedBone = TwoBoneIKGraphNode->Node.IKBone.BoneName;
+		}
+
+	}
+	else if (BoneSelectMode == BSM_JointTarget)
+	{
+		SelectedBone = TwoBoneIKGraphNode->Node.JointTargetSpaceBoneName;
+	}
+
+	return SelectedBone;
+}
+
+bool FTwoBoneIKEditMode::HandleClick(FEditorViewportClient* InViewportClient, HHitProxy* HitProxy, const FViewportClick& Click)
+{
+	bool bResult = FAnimNodeEditMode::HandleClick(InViewportClient, HitProxy, Click);
+
+	if (HitProxy != nullptr && HitProxy->IsA(HActor::StaticGetType()))
+	{
+		HActor* ActorHitProxy = static_cast<HActor*>(HitProxy);
+		ABoneSelectActor* BoneActor = Cast<ABoneSelectActor>(ActorHitProxy->Actor);
+
+		if (BoneActor)
+		{
+			// toggle bone selection mode
+			if (BoneSelectMode == BSM_EndEffector)
+			{
+				BoneSelectMode = BSM_JointTarget;
+			}
+			else
+			{
+				BoneSelectMode = BSM_EndEffector;
+			}
+
+			bResult = true;
+		}
+	}
+
+	return bResult;
+}
+
+void FTwoBoneIKEditMode::Tick(FEditorViewportClient* ViewportClient, float DeltaTime)
+{
+	FAnimNodeEditMode::Tick(ViewportClient, DeltaTime);
+
+	ManageBoneSelectActor();
+
+	USkeletalMeshComponent* SkelComp = GetAnimPreviewScene().GetPreviewMeshComponent();
+
+	if (BoneSelectActor.IsValid())
+	{
+		// move the actor's position
+		if (BoneSelectMode == BSM_JointTarget)
+		{
+			FName BoneName;
+			int32 EffectorBoneIndex = SkelComp->GetBoneIndex(TwoBoneIKGraphNode->Node.EffectorSpaceBoneName);
+			if (EffectorBoneIndex != INDEX_NONE)
+			{
+				BoneName = TwoBoneIKGraphNode->Node.EffectorSpaceBoneName;
+			}
+			else
+			{
+				BoneName = TwoBoneIKGraphNode->Node.IKBone.BoneName;
+			}
+
+			FVector ActorLocation = ConvertWidgetLocation(SkelComp, TwoBoneIKRuntimeNode->ForwardedPose, BoneName, TwoBoneIKGraphNode->Node.EffectorLocation, TwoBoneIKGraphNode->Node.EffectorLocationSpace);
+			BoneSelectActor->SetActorLocation(ActorLocation + FVector(0, 10, 0));
+
+		}
+		else if (BoneSelectMode == BSM_EndEffector)
+		{
+			int32 JointTargetIndex = SkelComp->GetBoneIndex(TwoBoneIKGraphNode->Node.JointTargetSpaceBoneName);
+
+			if (JointTargetIndex != INDEX_NONE)
+			{
+				FVector ActorLocation = ConvertWidgetLocation(SkelComp, TwoBoneIKRuntimeNode->ForwardedPose, TwoBoneIKGraphNode->Node.JointTargetSpaceBoneName, TwoBoneIKGraphNode->Node.JointTargetLocation, TwoBoneIKGraphNode->Node.JointTargetLocationSpace);
+				BoneSelectActor->SetActorLocation(ActorLocation + FVector(0, 10, 0));
+			}
+		}
+	}
+}
+
+void FTwoBoneIKEditMode::DoTranslation(FVector& InTranslation)
+{
+	USkeletalMeshComponent* SkelComp = GetAnimPreviewScene().GetPreviewMeshComponent();
+
+	if (BoneSelectMode == BSM_EndEffector)
+	{
+		FVector Offset = ConvertCSVectorToBoneSpace(SkelComp, InTranslation, TwoBoneIKRuntimeNode->ForwardedPose, GetSelectedBone(), TwoBoneIKGraphNode->Node.EffectorLocationSpace);
+
+		TwoBoneIKRuntimeNode->EffectorLocation += Offset;
+		TwoBoneIKGraphNode->Node.EffectorLocation = TwoBoneIKRuntimeNode->EffectorLocation;
+		TwoBoneIKGraphNode->SetDefaultValue(GET_MEMBER_NAME_STRING_CHECKED(FAnimNode_TwoBoneIK, EffectorLocation), TwoBoneIKRuntimeNode->EffectorLocation);
+	}
+	else if (BoneSelectMode == BSM_JointTarget)
+	{
+		FVector Offset = ConvertCSVectorToBoneSpace(SkelComp, InTranslation, TwoBoneIKRuntimeNode->ForwardedPose, GetSelectedBone(), TwoBoneIKGraphNode->Node.JointTargetLocationSpace);
+
+		TwoBoneIKRuntimeNode->JointTargetLocation += Offset;
+		TwoBoneIKGraphNode->Node.JointTargetLocation = TwoBoneIKRuntimeNode->JointTargetLocation;
+		TwoBoneIKGraphNode->SetDefaultValue(GET_MEMBER_NAME_STRING_CHECKED(FAnimNode_TwoBoneIK, JointTargetLocation), TwoBoneIKRuntimeNode->JointTargetLocation);
+	}
+}
+
+void FTwoBoneIKEditMode::ManageBoneSelectActor()
+{
+	USkeletalMeshComponent* SkelComp = GetAnimPreviewScene().GetPreviewMeshComponent();
+	USkeleton* Skeleton = SkelComp->SkeletalMesh->Skeleton;
+
+	// create a bone select actor if required
+	int32 JointTargetIndex = Skeleton->GetReferenceSkeleton().FindBoneIndex(TwoBoneIKGraphNode->Node.JointTargetSpaceBoneName);
+
+	const bool bNeedsActor = JointTargetIndex != INDEX_NONE &&
+							(TwoBoneIKGraphNode->Node.JointTargetLocationSpace == EBoneControlSpace::BCS_BoneSpace ||
+							 TwoBoneIKGraphNode->Node.JointTargetLocationSpace == EBoneControlSpace::BCS_ParentBoneSpace);
+
+	if (bNeedsActor)
+	{
+		if (!BoneSelectActor.IsValid())
+		{
+			UWorld* World = SkelComp->GetWorld();
+			check(World);
+
+			BoneSelectActor = World->SpawnActor<ABoneSelectActor>(FVector(0), FRotator(0, 0, 0));
+			check(BoneSelectActor.IsValid());
+		}
+	}
+	else if(BoneSelectActor.IsValid())
+	{
+		if (SkelComp->GetWorld()->DestroyActor(BoneSelectActor.Get()))
+		{
+			BoneSelectActor = nullptr;
+		}
+	}
+}
\ No newline at end of file
diff --git a/Engine/Source/Editor/AnimGraph/Private/EditModes/TwoBoneIKEditMode.h b/Engine/Source/Editor/AnimGraph/Private/EditModes/TwoBoneIKEditMode.h
new file mode 100644
--- /dev/null
+++ b/Engine/Source/Editor/AnimGraph/Private/EditModes/TwoBoneIKEditMode.h
@@ -0,0 +1,51 @@
+// Copyright 1998-2016 Epic Games, Inc. All Rights Reserved.
+
+#pragma once
+
+#include "AnimNodeEditMode.h"
+#include "AnimGraphNode_TwoBoneIK.h"
+
+class FTwoBoneIKEditMode : public FAnimNodeEditMode
+{
+public:
+	FTwoBoneIKEditMode();
+
+	/** IAnimNodeEditMode interface */
+	virtual void EnterMode(class UAnimGraphNode_Base* InEditorNode, struct FAnimNode_Base* InRuntimeNode) override;
+	virtual void ExitMode() override;
+	virtual FVector GetWidgetLocation() const override;
+	virtual FWidget::EWidgetMode GetWidgetMode() const override;
+	virtual FName GetSelectedBone() const override;
+	virtual void DoTranslation(FVector& InTranslation) override;
+
+	/** FEdMode interface */
+	virtual void Render(const FSceneView* View, FViewport* Viewport, FPrimitiveDrawInterface* PDI) override;
+	virtual bool HandleClick(FEditorViewportClient* InViewportClient, HHitProxy* HitProxy, const FViewportClick& Click) override;
+	virtual void Tick(FEditorViewportClient* ViewportClient, float DeltaTime) override;
+
+private:
+	/** Helper function for Render() */
+	void DrawTargetLocation(FPrimitiveDrawInterface* PDI, USkeletalMeshComponent* SkelComp, USkeleton * Skeleton, EBoneControlSpace SpaceBase, FName SpaceBoneName, const FVector& TargetLocation, const FColor& TargetColor, const FColor& BoneColor) const;
+
+	/** Creates and destroys the selection actor as needed */
+	void ManageBoneSelectActor();
+
+private:
+	/** Cache the typed nodes */
+	struct FAnimNode_TwoBoneIK* TwoBoneIKRuntimeNode;
+	UAnimGraphNode_TwoBoneIK* TwoBoneIKGraphNode;
+
+	/** Bone selection mode */
+	enum BoneSelectModeType
+	{
+		BSM_EndEffector,
+		BSM_JointTarget,
+		BSM_Max
+	};
+
+	/** The current bone selection mode */
+	BoneSelectModeType BoneSelectMode;
+
+	/** */
+	TWeakObjectPtr<class ABoneSelectActor> BoneSelectActor;
+};
\ No newline at end of file
diff --git a/Engine/Source/Editor/AnimGraph/Public/AnimGraphModule.h b/Engine/Source/Editor/AnimGraph/Public/AnimGraphModule.h
--- a/Engine/Source/Editor/AnimGraph/Public/AnimGraphModule.h
+++ b/Engine/Source/Editor/AnimGraph/Public/AnimGraphModule.h
@@ -9,8 +9,7 @@
 class FAnimGraphModule : public IModuleInterface
 {
 public:
-	virtual void StartupModule() override
-	{
-		FAnimGraphCommands::Register();
-	}
+	/** IModuleInterface interface */
+	virtual void StartupModule() override;
+	virtual void ShutdownModule() override;
 };
\ No newline at end of file
diff --git a/Engine/Source/Editor/AnimGraph/Public/AnimGraphNode_ObserveBone.h b/Engine/Source/Editor/AnimGraph/Public/AnimGraphNode_ObserveBone.h
--- a/Engine/Source/Editor/AnimGraph/Public/AnimGraphNode_ObserveBone.h
+++ b/Engine/Source/Editor/AnimGraph/Public/AnimGraphNode_ObserveBone.h
@@ -29,14 +29,11 @@ protected:
 	virtual void ValidateAnimNodeDuringCompilation(USkeleton* ForSkeleton, FCompilerResultsLog& MessageLog) override;
 	virtual void CustomizePinData(UEdGraphPin* Pin, FName SourcePropertyName, int32 ArrayIndex) const override;
 	virtual void CustomizeDetails(IDetailLayoutBuilder& DetailBuilder) override;
+	virtual FEditorModeID GetEditorMode() const override;
 	// End of UAnimGraphNode_Base interface
 
 	// UAnimGraphNode_SkeletalControlBase interface
 	virtual FText GetControllerDescription() const override;
-	virtual int32 GetWidgetCoordinateSystem(const USkeletalMeshComponent* SkelComp) override;
-	virtual FVector GetWidgetLocation(const USkeletalMeshComponent* SkelComp, struct FAnimNode_SkeletalControlBase* AnimNode) override;
-	virtual int32 GetWidgetMode(const USkeletalMeshComponent* SkelComp) override;
-	virtual FName FindSelectedBone() override;
 	virtual const FAnimNode_SkeletalControlBase* GetNode() const override { return &Node; }
 	// End of UAnimGraphNode_SkeletalControlBase interface
 };
diff --git a/Engine/Source/Editor/AnimGraph/Public/AnimNodeEditMode.h b/Engine/Source/Editor/AnimGraph/Public/AnimNodeEditMode.h
new file mode 100644
--- /dev/null
+++ b/Engine/Source/Editor/AnimGraph/Public/AnimNodeEditMode.h
@@ -0,0 +1,66 @@
+// Copyright 1998-2016 Epic Games, Inc. All Rights Reserved.
+
+#pragma once
+
+#include "IAnimNodeEditMode.h"
+#include "BonePose.h"
+
+/** Base implementation for anim node edit modes */
+class FAnimNodeEditMode : public IAnimNodeEditMode
+{
+public:
+	FAnimNodeEditMode();
+
+	/** IAnimNodeEditMode interface */
+	virtual ECoordSystem GetWidgetCoordinateSystem() const override;
+	virtual FWidget::EWidgetMode GetWidgetMode() const override;
+	virtual FWidget::EWidgetMode ChangeToNextWidgetMode(FWidget::EWidgetMode CurWidgetMode) override;
+	virtual bool SetWidgetMode(FWidget::EWidgetMode InWidgetMode) override;
+	virtual FName GetSelectedBone() const override;
+	virtual void DoTranslation(FVector& InTranslation) override;
+	virtual void DoRotation(FRotator& InRotation) override;
+	virtual void DoScale(FVector& InScale) override;
+	virtual void EnterMode(class UAnimGraphNode_Base* InEditorNode, struct FAnimNode_Base* InRuntimeNode) override;
+	virtual void ExitMode() override;
+
+	/** IPersonaEditMode interface */
+	virtual bool GetCameraTarget(FSphere& OutTarget) const override;
+	virtual class IPersonaPreviewScene& GetAnimPreviewScene() const override;
+	virtual void GetOnScreenDebugInfo(TArray<FText>& OutDebugInfo) const override;
+
+	/** FEdMode interface */
+	virtual void Render(const FSceneView* View, FViewport* Viewport, FPrimitiveDrawInterface* PDI) override;
+	virtual void DrawHUD(FEditorViewportClient* ViewportClient, FViewport* Viewport, const FSceneView* View, FCanvas* Canvas) override;
+	virtual bool HandleClick(FEditorViewportClient* InViewportClient, HHitProxy* HitProxy, const FViewportClick& Click) override;
+	virtual FVector GetWidgetLocation() const override;
+	virtual bool StartTracking(FEditorViewportClient* InViewportClient, FViewport* InViewport) override;
+	virtual bool EndTracking(FEditorViewportClient* InViewportClient, FViewport* InViewport) override;
+	virtual bool InputKey(FEditorViewportClient* InViewportClient, FViewport* InViewport, FKey InKey, EInputEvent InEvent) override;
+	virtual bool InputDelta(FEditorViewportClient* InViewportClient, FViewport* InViewport, FVector& InDrag, FRotator& InRot, FVector& InScale) override;
+	virtual bool GetCustomDrawingCoordinateSystem(FMatrix& InMatrix, void* InData) override;
+	virtual bool GetCustomInputCoordinateSystem(FMatrix& InMatrix, void* InData) override;
+	virtual bool ShouldDrawWidget() const override;
+	virtual void Tick(FEditorViewportClient* ViewportClient, float DeltaTime) override;
+
+protected:
+	// local conversion function for drawing
+	static void ConvertToComponentSpaceTransform(const USkeletalMeshComponent* SkelComp, const FTransform & InTransform, FTransform & OutCSTransform, int32 BoneIndex, EBoneControlSpace Space);
+	// convert drag vector in component space to bone space 
+	static FVector ConvertCSVectorToBoneSpace(const USkeletalMeshComponent* SkelComp, FVector& InCSVector, FCSPose<FCompactHeapPose>& MeshBases, const FName& BoneName, const EBoneControlSpace Space);
+	// convert rotator in component space to bone space 
+	static FQuat ConvertCSRotationToBoneSpace(const USkeletalMeshComponent* SkelComp, FRotator& InCSRotator, FCSPose<FCompactHeapPose>& MeshBases, const FName& BoneName, const EBoneControlSpace Space);
+	// convert widget location according to bone control space
+	static FVector ConvertWidgetLocation(const USkeletalMeshComponent* InSkelComp, FCSPose<FCompactHeapPose>& InMeshBases, const FName& BoneName, const FVector& InLocation, const EBoneControlSpace Space);
+
+protected:
+	/** The node we are operating on */
+	class UAnimGraphNode_Base* AnimNode;
+
+	/** The runtime node in the preview scene */
+	struct FAnimNode_Base* RuntimeAnimNode;
+
+private:
+	bool bManipulating;
+
+	bool bInTransaction;
+};
\ No newline at end of file
diff --git a/Engine/Source/Editor/AnimGraph/Public/AnimNodeEditModes.h b/Engine/Source/Editor/AnimGraph/Public/AnimNodeEditModes.h
new file mode 100644
--- /dev/null
+++ b/Engine/Source/Editor/AnimGraph/Public/AnimNodeEditModes.h
@@ -0,0 +1,13 @@
+// Copyright 1998-2016 Epic Games, Inc. All Rights Reserved.
+
+#pragma once
+
+struct ANIMGRAPH_API AnimNodeEditModes
+{
+	const static FEditorModeID AnimNode;
+	const static FEditorModeID TwoBoneIK;
+	const static FEditorModeID ObserveBone;
+	const static FEditorModeID ModifyBone;
+	const static FEditorModeID Fabrik;
+	const static FEditorModeID PoseDriver;
+};
\ No newline at end of file
diff --git a/Engine/Source/Editor/AnimGraph/Public/IAnimNodeEditMode.h b/Engine/Source/Editor/AnimGraph/Public/IAnimNodeEditMode.h
new file mode 100644
--- /dev/null
+++ b/Engine/Source/Editor/AnimGraph/Public/IAnimNodeEditMode.h
@@ -0,0 +1,40 @@
+// Copyright 1998-2016 Epic Games, Inc. All Rights Reserved.
+
+#pragma once
+
+#include "IPersonaEditMode.h"
+
+/** Base interface for skeletal control edit modes */
+class IAnimNodeEditMode : public IPersonaEditMode
+{
+public:
+	/** Returns the coordinate system that should be used for this bone */
+	virtual ECoordSystem GetWidgetCoordinateSystem() const = 0;
+
+	/** @return current widget mode this anim graph node supports */
+	virtual FWidget::EWidgetMode GetWidgetMode() const = 0;
+
+	/** Called when the user changed widget mode by pressing "Space" key */
+	virtual FWidget::EWidgetMode ChangeToNextWidgetMode(FWidget::EWidgetMode CurWidgetMode) = 0;
+
+	/** Called when the user set widget mode directly, returns true if InWidgetMode is available */
+	virtual bool SetWidgetMode(FWidget::EWidgetMode InWidgetMode) = 0;
+
+	/** Get the bone that the skeletal control is manipulating */
+	virtual FName GetSelectedBone() const = 0;
+
+	/** Called when the widget is dragged in translation mode */
+	virtual void DoTranslation(FVector& InTranslation) = 0;
+
+	/** Called when the widget is dragged in rotation mode */
+	virtual void DoRotation(FRotator& InRotation) = 0;
+
+	/** Called when the widget is dragged in scale mode */
+	virtual void DoScale(FVector& InScale) = 0;
+
+	/** Called when entering this edit mode */
+	virtual void EnterMode(class UAnimGraphNode_Base* InEditorNode, struct FAnimNode_Base* InRuntimeNode) = 0;
+
+	/** Called when exiting this edit mode */
+	virtual void ExitMode() = 0;
+};
\ No newline at end of file
diff --git a/Engine/Source/Editor/AnimationBlueprintEditor/Private/AnimationBlueprintEditorMode.cpp b/Engine/Source/Editor/AnimationBlueprintEditor/Private/AnimationBlueprintEditorMode.cpp
new file mode 100644
--- /dev/null
+++ b/Engine/Source/Editor/AnimationBlueprintEditor/Private/AnimationBlueprintEditorMode.cpp
@@ -0,0 +1,175 @@
+// Copyright 1998-2016 Epic Games, Inc. All Rights Reserved.
+
+#include "AnimationBlueprintEditorPrivatePCH.h"
+
+#include "Animation/AnimInstance.h"
+
+#include "AnimationBlueprintEditor.h"
+
+#include "AnimationBlueprintEditorMode.h"
+#include "IPersonaToolkit.h"
+#include "ISkeletonTree.h"
+#include "BlueprintEditorTabs.h"
+#include "ISkeletonEditorModule.h"
+#include "PersonaModule.h"
+#include "SBlueprintEditorToolbar.h"
+#include "IPersonaPreviewScene.h"
+/////////////////////////////////////////////////////
+// FAnimationBlueprintEditorMode
+
+FAnimationBlueprintEditorMode::FAnimationBlueprintEditorMode(const TSharedRef<FAnimationBlueprintEditor>& InAnimationBlueprintEditor)
+	: FBlueprintEditorApplicationMode(InAnimationBlueprintEditor, FAnimationBlueprintEditorModes::AnimationBlueprintEditorMode, FAnimationBlueprintEditorModes::GetLocalizedMode, false, false)
+{
+	PreviewScenePtr = InAnimationBlueprintEditor->GetPreviewScene();
+	AnimBlueprintPtr = CastChecked<UAnimBlueprint>(InAnimationBlueprintEditor->GetBlueprintObj());
+
+	TabLayout = FTabManager::NewLayout( "Stanalone_AnimationBlueprintEditMode_Layout_v1.1" )
+		->AddArea
+		(
+			FTabManager::NewPrimaryArea()
+			->SetOrientation(Orient_Vertical)
+			->Split
+			(
+				// Top toolbar
+				FTabManager::NewStack() 
+				->SetSizeCoefficient(0.186721f)
+				->SetHideTabWell(true)
+				->AddTab(InAnimationBlueprintEditor->GetToolbarTabId(), ETabState::OpenedTab )
+			)
+			->Split
+			(
+				// Main application area
+				FTabManager::NewSplitter()
+				->SetOrientation(Orient_Horizontal)
+				->Split
+				(
+					// Left side
+					FTabManager::NewSplitter()
+					->SetSizeCoefficient(0.25f)
+					->SetOrientation(Orient_Vertical)
+					->Split
+					(
+						// Left top - viewport
+						FTabManager::NewStack()
+						->SetSizeCoefficient(0.5f)
+						->SetHideTabWell(true)
+						->AddTab(AnimationBlueprintEditorTabs::ViewportTab, ETabState::OpenedTab)
+					)
+					->Split
+					(
+						//	Left bottom - preview settings
+						FTabManager::NewStack()
+						->SetSizeCoefficient(0.5f)
+						->AddTab(AnimationBlueprintEditorTabs::SkeletonTreeTab, ETabState::ClosedTab)
+						->AddTab(FBlueprintEditorTabs::MyBlueprintID, ETabState::OpenedTab)
+					)
+				)
+				->Split
+				(
+					// Middle 
+					FTabManager::NewSplitter()
+					->SetOrientation(Orient_Vertical)
+					->SetSizeCoefficient(0.55f)
+					->Split
+					(
+						// Middle top - document edit area
+						FTabManager::NewStack()
+						->SetSizeCoefficient(0.8f)
+						->AddTab("Document", ETabState::ClosedTab)
+					)
+					->Split
+					(
+						// Middle bottom - compiler results & find
+						FTabManager::NewStack()
+						->SetSizeCoefficient(0.2f)
+						->AddTab(FBlueprintEditorTabs::CompilerResultsID, ETabState::ClosedTab)
+						->AddTab(FBlueprintEditorTabs::FindResultsID, ETabState::ClosedTab)
+					)
+				)
+				->Split
+				(
+					// Right side
+					FTabManager::NewSplitter()
+					->SetSizeCoefficient(0.2f)
+					->SetOrientation(Orient_Vertical)
+					->Split
+					(
+						// Right top - selection details panel & overrides
+						FTabManager::NewStack()
+						->SetHideTabWell(false)
+						->SetSizeCoefficient(0.5f)
+						->AddTab(FBlueprintEditorTabs::DetailsID, ETabState::OpenedTab)
+						->AddTab(AnimationBlueprintEditorTabs::AdvancedPreviewTab, ETabState::ClosedTab)
+						->AddTab(AnimationBlueprintEditorTabs::AssetOverridesTab, ETabState::ClosedTab)
+						->SetForegroundTab(FBlueprintEditorTabs::DetailsID)
+					)
+					->Split
+					(
+						// Right bottom - Asset browser & advanced preview settings
+						FTabManager::NewStack()
+						->SetHideTabWell(false)
+						->SetSizeCoefficient(0.5f)
+						->AddTab(AnimationBlueprintEditorTabs::AnimBlueprintPreviewEditorTab, ETabState::OpenedTab)
+						->AddTab(AnimationBlueprintEditorTabs::AssetBrowserTab, ETabState::OpenedTab)
+						->AddTab(AnimationBlueprintEditorTabs::SlotNamesTab, ETabState::ClosedTab)
+						->SetForegroundTab(AnimationBlueprintEditorTabs::AnimBlueprintPreviewEditorTab)
+					)
+				)
+			)
+		);
+
+	ISkeletonEditorModule& SkeletonEditorModule = FModuleManager::LoadModuleChecked<ISkeletonEditorModule>("SkeletonEditor");
+	TabFactories.RegisterFactory(SkeletonEditorModule.CreateSkeletonTreeTabFactory(InAnimationBlueprintEditor, InAnimationBlueprintEditor->GetSkeletonTree()));
+
+	FPersonaModule& PersonaModule = FModuleManager::LoadModuleChecked<FPersonaModule>("Persona");
+	TabFactories.RegisterFactory(PersonaModule.CreatePersonaViewportTabFactory(InAnimationBlueprintEditor, InAnimationBlueprintEditor->GetSkeletonTree(), InAnimationBlueprintEditor->GetPersonaToolkit()->GetPreviewScene(), InAnimationBlueprintEditor->OnPostUndo, InAnimationBlueprintEditor, FOnViewportCreated(), true, false));
+	TabFactories.RegisterFactory(PersonaModule.CreateAdvancedPreviewSceneTabFactory(InAnimationBlueprintEditor, InAnimationBlueprintEditor->GetPersonaToolkit()->GetPreviewScene()));
+	TabFactories.RegisterFactory(PersonaModule.CreateAnimationAssetBrowserTabFactory(InAnimationBlueprintEditor, InAnimationBlueprintEditor->GetPersonaToolkit(), FOnOpenNewAsset::CreateSP(&InAnimationBlueprintEditor.Get(), &FAnimationBlueprintEditor::HandleOpenNewAsset), FOnAnimationSequenceBrowserCreated(), true));
+	TabFactories.RegisterFactory(PersonaModule.CreateAnimBlueprintPreviewTabFactory(InAnimationBlueprintEditor, InAnimationBlueprintEditor->GetPersonaToolkit()->GetPreviewScene()));
+	TabFactories.RegisterFactory(PersonaModule.CreateAnimBlueprintAssetOverridesTabFactory(InAnimationBlueprintEditor, InAnimationBlueprintEditor->GetPersonaToolkit()->GetAnimBlueprint(), InAnimationBlueprintEditor->OnPostUndo));
+	TabFactories.RegisterFactory(PersonaModule.CreateSkeletonSlotNamesTabFactory(InAnimationBlueprintEditor, InAnimationBlueprintEditor->GetSkeletonTree()->GetEditableSkeleton(), InAnimationBlueprintEditor->OnPostUndo, FOnObjectSelected::CreateSP(&InAnimationBlueprintEditor.Get(), &FAnimationBlueprintEditor::HandleObjectSelected)));
+
+	// setup toolbar - clear existing toolbar extender from the BP mode
+	//@TODO: Keep this in sync with BlueprintEditorModes.cpp
+	ToolbarExtender = MakeShareable(new FExtender);
+	InAnimationBlueprintEditor->GetToolbarBuilder()->AddCompileToolbar(ToolbarExtender);
+	InAnimationBlueprintEditor->GetToolbarBuilder()->AddScriptingToolbar(ToolbarExtender);
+	InAnimationBlueprintEditor->GetToolbarBuilder()->AddBlueprintGlobalOptionsToolbar(ToolbarExtender);
+	InAnimationBlueprintEditor->GetToolbarBuilder()->AddDebuggingToolbar(ToolbarExtender);
+}
+
+void FAnimationBlueprintEditorMode::RegisterTabFactories(TSharedPtr<FTabManager> InTabManager)
+{
+	TSharedPtr<FBlueprintEditor> BP = MyBlueprintEditor.Pin();
+	
+	BP->RegisterToolbarTab(InTabManager.ToSharedRef());
+
+	// Mode-specific setup
+	BP->PushTabFactories(CoreTabFactories);
+	BP->PushTabFactories(BlueprintEditorTabFactories);
+	BP->PushTabFactories(TabFactories);
+}
+
+void FAnimationBlueprintEditorMode::PostActivateMode()
+{
+	if (UAnimBlueprint* AnimBlueprint = AnimBlueprintPtr.Get())
+	{
+		// Switch off any active preview when going to graph editing mode
+		PreviewScenePtr.Pin()->SetPreviewAnimationAsset(NULL, false);
+
+		// When switching to anim blueprint mode, make sure the object being debugged is either a valid world object or the preview instance
+		UDebugSkelMeshComponent* PreviewComponent = PreviewScenePtr.Pin()->GetPreviewMeshComponent();
+		if ((AnimBlueprint->GetObjectBeingDebugged() == NULL) && (PreviewComponent->IsAnimBlueprintInstanced()))
+		{
+			AnimBlueprint->SetObjectBeingDebugged(PreviewComponent->GetAnimInstance());
+		}
+
+		// If we are a derived anim blueprint always show the overrides tab
+		if(UAnimBlueprint::FindRootAnimBlueprint(AnimBlueprint))
+		{
+			MyBlueprintEditor.Pin()->GetTabManager()->InvokeTab(AnimationBlueprintEditorTabs::AssetOverridesTab);
+		}
+	}
+
+	FBlueprintEditorApplicationMode::PostActivateMode();
+}
\ No newline at end of file
diff --git a/Engine/Source/Editor/AnimationBlueprintEditor/Private/AnimationBlueprintEditorPrivatePCH.h b/Engine/Source/Editor/AnimationBlueprintEditor/Private/AnimationBlueprintEditorPrivatePCH.h
new file mode 100644
--- /dev/null
+++ b/Engine/Source/Editor/AnimationBlueprintEditor/Private/AnimationBlueprintEditorPrivatePCH.h
@@ -0,0 +1,15 @@
+// Copyright 1998-2016 Epic Games, Inc. All Rights Reserved.
+
+#pragma once
+
+#include "UnrealEd.h"
+#include "SlateBasics.h"
+#include "NotificationManager.h"
+#include "SNotificationList.h"
+
+#include "Editor/Kismet/Public/WorkflowOrientedApp/ApplicationMode.h"
+#include "Editor/Kismet/Public/WorkflowOrientedApp/WorkflowTabFactory.h"
+#include "Editor/Kismet/Public/WorkflowOrientedApp/WorkflowTabManager.h"
+
+#include "PropertyEditing.h"
+#include "PropertyCustomizationHelpers.h"
diff --git a/Engine/Source/Editor/AnimationEditor/Private/AnimationEditor.cpp b/Engine/Source/Editor/AnimationEditor/Private/AnimationEditor.cpp
new file mode 100644
--- /dev/null
+++ b/Engine/Source/Editor/AnimationEditor/Private/AnimationEditor.cpp
@@ -0,0 +1,866 @@
+// Copyright 1998-2016 Epic Games, Inc. All Rights Reserved.
+
+#include "AnimationEditorPrivatePCH.h"
+#include "IAnimationEditorModule.h"
+#include "AnimationEditor.h"
+#include "IPersonaToolkit.h"
+#include "PersonaModule.h"
+#include "AnimationEditorMode.h"
+#include "IPersonaPreviewScene.h"
+#include "AnimationEditorCommands.h"
+#include "IEditableSkeleton.h"
+#include "IDetailsView.h"
+#include "ISkeletonTree.h"
+#include "ISkeletonEditorModule.h"
+#include "IDocumentation.h"
+#include "SDockTab.h"
+#include "AnimPreviewInstance.h"
+#include "ScopedTransaction.h"
+#include "AnimationEditorUtils.h"
+#include "AssetRegistryModule.h"
+#include "IAssetFamily.h"
+#include "AssetEditorModeManager.h"
+#include "IAnimationSequenceBrowser.h"
+#include "Animation/PoseAsset.h"
+#include "SNotificationList.h"
+#include "NotificationManager.h"
+
+const FName AnimationEditorAppIdentifier = FName(TEXT("AnimationEditorApp"));
+
+const FName AnimationEditorModes::AnimationEditorMode(TEXT("AnimationEditorMode"));
+
+const FName AnimationEditorTabs::DetailsTab(TEXT("DetailsTab"));
+const FName AnimationEditorTabs::SkeletonTreeTab(TEXT("SkeletonTreeView"));
+const FName AnimationEditorTabs::ViewportTab(TEXT("Viewport"));
+const FName AnimationEditorTabs::AdvancedPreviewTab(TEXT("AdvancedPreviewTab"));
+const FName AnimationEditorTabs::DocumentTab(TEXT("Document"));
+const FName AnimationEditorTabs::AssetBrowserTab(TEXT("SequenceBrowser"));
+const FName AnimationEditorTabs::AssetDetailsTab(TEXT("AnimAssetPropertiesTab"));
+const FName AnimationEditorTabs::CurveNamesTab(TEXT("AnimCurveViewerTab"));
+const FName AnimationEditorTabs::SlotNamesTab(TEXT("SkeletonSlotNames"));
+
+DEFINE_LOG_CATEGORY(LogAnimationEditor);
+
+#define LOCTEXT_NAMESPACE "AnimationEditor"
+
+FAnimationEditor::FAnimationEditor()
+{
+	UEditorEngine* Editor = Cast<UEditorEngine>(GEngine);
+	if (Editor != nullptr)
+	{
+		Editor->RegisterForUndo(this);
+	}
+}
+
+FAnimationEditor::~FAnimationEditor()
+{
+	UEditorEngine* Editor = Cast<UEditorEngine>(GEngine);
+	if (Editor != nullptr)
+	{
+		Editor->UnregisterForUndo(this);
+	}
+
+	FEditorDelegates::OnAssetPostImport.RemoveAll(this);
+	FReimportManager::Instance()->OnPostReimport().RemoveAll(this);
+}
+
+void FAnimationEditor::RegisterTabSpawners(const TSharedRef<class FTabManager>& InTabManager)
+{
+	WorkspaceMenuCategory = InTabManager->AddLocalWorkspaceMenuCategory(LOCTEXT("WorkspaceMenu_AnimationEditor", "Animation Editor"));
+
+	FAssetEditorToolkit::RegisterTabSpawners( InTabManager );
+}
+
+void FAnimationEditor::UnregisterTabSpawners(const TSharedRef<class FTabManager>& InTabManager)
+{
+	FAssetEditorToolkit::UnregisterTabSpawners(InTabManager);
+}
+
+void FAnimationEditor::InitAnimationEditor(const EToolkitMode::Type Mode, const TSharedPtr<IToolkitHost>& InitToolkitHost, UAnimationAsset* InAnimationAsset)
+{
+	AnimationAsset = InAnimationAsset;
+
+	// Register post import callback to catch animation imports when we have the asset open (we need to reinit)
+	FReimportManager::Instance()->OnPostReimport().AddRaw(this, &FAnimationEditor::HandlePostReimport);
+	FEditorDelegates::OnAssetPostImport.AddRaw(this, &FAnimationEditor::HandlePostImport);
+
+	FPersonaModule& PersonaModule = FModuleManager::LoadModuleChecked<FPersonaModule>("Persona");
+	PersonaToolkit = PersonaModule.CreatePersonaToolkit(InAnimationAsset);
+
+	PersonaToolkit->GetPreviewScene()->SetDefaultAnimationMode(EPreviewSceneDefaultAnimationMode::Animation);
+
+	FSkeletonTreeArgs SkeletonTreeArgs(OnPostUndo);
+	SkeletonTreeArgs.OnObjectSelected = FOnObjectSelected::CreateSP(this, &FAnimationEditor::HandleObjectSelected);
+	SkeletonTreeArgs.PreviewScene = PersonaToolkit->GetPreviewScene();
+
+	ISkeletonEditorModule& SkeletonEditorModule = FModuleManager::GetModuleChecked<ISkeletonEditorModule>("SkeletonEditor");
+	SkeletonTree = SkeletonEditorModule.CreateSkeletonTree(PersonaToolkit->GetSkeleton(), PersonaToolkit->GetMesh(), SkeletonTreeArgs);
+
+	const bool bCreateDefaultStandaloneMenu = true;
+	const bool bCreateDefaultToolbar = true;
+	const TSharedRef<FTabManager::FLayout> DummyLayout = FTabManager::NewLayout("NullLayout")->AddArea(FTabManager::NewPrimaryArea());
+	FAssetEditorToolkit::InitAssetEditor(Mode, InitToolkitHost, AnimationEditorAppIdentifier, DummyLayout, bCreateDefaultStandaloneMenu, bCreateDefaultToolbar, InAnimationAsset);
+
+	BindCommands();
+
+	AddApplicationMode(
+		AnimationEditorModes::AnimationEditorMode,
+		MakeShareable(new FAnimationEditorMode(SharedThis(this), SkeletonTree.ToSharedRef())));
+
+	SetCurrentMode(AnimationEditorModes::AnimationEditorMode);
+
+	// set up our editor mode
+	check(AssetEditorModeManager);
+	AssetEditorModeManager->SetDefaultMode(FPersonaEditModes::SkeletonSelection);
+
+	ExtendMenu();
+	ExtendToolbar();
+	RegenerateMenusAndToolbars();
+
+	OpenNewAnimationDocumentTab(AnimationAsset);
+}
+
+FName FAnimationEditor::GetToolkitFName() const
+{
+	return FName("AnimationEditor");
+}
+
+FText FAnimationEditor::GetBaseToolkitName() const
+{
+	return LOCTEXT("AppLabel", "AnimationEditor");
+}
+
+FString FAnimationEditor::GetWorldCentricTabPrefix() const
+{
+	return LOCTEXT("WorldCentricTabPrefix", "AnimationEditor ").ToString();
+}
+
+FLinearColor FAnimationEditor::GetWorldCentricTabColorScale() const
+{
+	return FLinearColor(0.3f, 0.2f, 0.5f, 0.5f);
+}
+
+void FAnimationEditor::AddReferencedObjects(FReferenceCollector& Collector)
+{
+	Collector.AddReferencedObject(AnimationAsset);
+}
+
+void FAnimationEditor::BindCommands()
+{
+	FAnimationEditorCommands::Register();
+	
+	ToolkitCommands->MapAction(FAnimationEditorCommands::Get().ApplyCompression,
+		FExecuteAction::CreateSP(this, &FAnimationEditor::OnApplyCompression),
+		FCanExecuteAction::CreateSP(this, &FAnimationEditor::HasValidAnimationSequence));
+
+	ToolkitCommands->MapAction(FAnimationEditorCommands::Get().SetKey,
+		FExecuteAction::CreateSP(this, &FAnimationEditor::OnSetKey),
+		FCanExecuteAction::CreateSP(this, &FAnimationEditor::CanSetKey));
+
+	ToolkitCommands->MapAction(FAnimationEditorCommands::Get().ApplyAnimation,
+		FExecuteAction::CreateSP(this, &FAnimationEditor::OnApplyRawAnimChanges),
+		FCanExecuteAction::CreateSP(this, &FAnimationEditor::CanApplyRawAnimChanges));
+
+	ToolkitCommands->MapAction(FAnimationEditorCommands::Get().ExportToFBX,
+		FExecuteAction::CreateSP(this, &FAnimationEditor::OnExportToFBX),
+		FCanExecuteAction::CreateSP(this, &FAnimationEditor::HasValidAnimationSequence));
+
+	ToolkitCommands->MapAction(FAnimationEditorCommands::Get().AddLoopingInterpolation,
+		FExecuteAction::CreateSP(this, &FAnimationEditor::OnAddLoopingInterpolation),
+		FCanExecuteAction::CreateSP(this, &FAnimationEditor::HasValidAnimationSequence));
+}
+
+void FAnimationEditor::ExtendToolbar()
+{
+	// If the ToolbarExtender is valid, remove it before rebuilding it
+	if (ToolbarExtender.IsValid())
+	{
+		RemoveToolbarExtender(ToolbarExtender);
+		ToolbarExtender.Reset();
+	}
+
+	ToolbarExtender = MakeShareable(new FExtender);
+
+	AddToolbarExtender(ToolbarExtender);
+
+	IAnimationEditorModule& AnimationEditorModule = FModuleManager::GetModuleChecked<IAnimationEditorModule>("AnimationEditor");
+	AddToolbarExtender(AnimationEditorModule.GetToolBarExtensibilityManager()->GetAllExtenders(GetToolkitCommands(), GetEditingObjects()));
+
+	TArray<IAnimationEditorModule::FAnimationEditorToolbarExtender> ToolbarExtenderDelegates = AnimationEditorModule.GetAllAnimationEditorToolbarExtenders();
+
+	for (auto& ToolbarExtenderDelegate : ToolbarExtenderDelegates)
+	{
+		if (ToolbarExtenderDelegate.IsBound())
+		{
+			AddToolbarExtender(ToolbarExtenderDelegate.Execute(GetToolkitCommands(), SharedThis(this)));
+		}
+	}
+
+	// extend extra menu/toolbars
+	struct Local
+	{
+		static void FillToolbar(FToolBarBuilder& ToolbarBuilder, FAnimationEditor* InAnimationEditor)
+		{
+			ToolbarBuilder.BeginSection("Animation");
+			{
+				// create button
+				{
+					ToolbarBuilder.AddComboButton(
+						FUIAction(),
+						FOnGetContent::CreateSP(InAnimationEditor, &FAnimationEditor::GenerateCreateAssetMenu),
+						LOCTEXT("CreateAsset_Label", "Create Asset"),
+						LOCTEXT("CreateAsset_ToolTip", "Create Assets for this skeleton."),
+						FSlateIcon(FEditorStyle::GetStyleSetName(), "Persona.CreateAsset")
+						);
+				}
+
+				ToolbarBuilder.AddToolBarButton(FAnimationEditorCommands::Get().ApplyCompression, NAME_None, LOCTEXT("Toolbar_ApplyCompression", "Compression"));
+			}
+			ToolbarBuilder.EndSection();
+
+			ToolbarBuilder.BeginSection("Editing");
+			{
+				ToolbarBuilder.AddToolBarButton(FAnimationEditorCommands::Get().SetKey, NAME_None, LOCTEXT("Toolbar_SetKey", "Key"));
+				ToolbarBuilder.AddToolBarButton(FAnimationEditorCommands::Get().ApplyAnimation, NAME_None, LOCTEXT("Toolbar_ApplyAnimation", "Apply"));
+			}
+			ToolbarBuilder.EndSection();
+		}
+	};
+
+	ToolbarExtender->AddToolBarExtension(
+		"Asset",
+		EExtensionHook::After,
+		GetToolkitCommands(),
+		FToolBarExtensionDelegate::CreateStatic(&Local::FillToolbar, this)
+		);
+
+	ToolbarExtender->AddToolBarExtension(
+		"Asset",
+		EExtensionHook::After,
+		GetToolkitCommands(),
+		FToolBarExtensionDelegate::CreateLambda([this](FToolBarBuilder& ParentToolbarBuilder)
+		{
+			FPersonaModule& PersonaModule = FModuleManager::LoadModuleChecked<FPersonaModule>("Persona");
+			TSharedRef<class IAssetFamily> AssetFamily = PersonaModule.CreatePersonaAssetFamily(AnimationAsset);
+			AddToolbarWidget(PersonaModule.CreateAssetFamilyShortcutWidget(SharedThis(this), AssetFamily));
+		}	
+	));
+}
+
+void FAnimationEditor::ExtendMenu()
+{
+	MenuExtender = MakeShareable(new FExtender);
+
+	struct Local
+	{
+		static void AddAssetMenu(FMenuBuilder& MenuBuilder)
+		{
+			MenuBuilder.BeginSection("AnimationEditor", LOCTEXT("AnimationEditorAssetMenu_Animation", "Animation"));
+			{
+				MenuBuilder.AddMenuEntry(FAnimationEditorCommands::Get().ApplyCompression);
+				MenuBuilder.AddMenuEntry(FAnimationEditorCommands::Get().ExportToFBX);
+				MenuBuilder.AddMenuEntry(FAnimationEditorCommands::Get().AddLoopingInterpolation);
+			}
+			MenuBuilder.EndSection();
+		}
+	};
+
+	MenuExtender->AddMenuExtension(
+		"AssetEditorActions",
+		EExtensionHook::After,
+		GetToolkitCommands(),
+		FMenuExtensionDelegate::CreateStatic(&Local::AddAssetMenu)
+		);
+
+	AddMenuExtender(MenuExtender);
+
+	IAnimationEditorModule& AnimationEditorModule = FModuleManager::GetModuleChecked<IAnimationEditorModule>("AnimationEditor");
+	AddMenuExtender(AnimationEditorModule.GetToolBarExtensibilityManager()->GetAllExtenders(GetToolkitCommands(), GetEditingObjects()));
+}
+
+void FAnimationEditor::HandleObjectsSelected(const TArray<UObject*>& InObjects)
+{
+	if (DetailsView.IsValid())
+	{
+		DetailsView->SetObjects(InObjects);
+	}
+}
+
+void FAnimationEditor::HandleObjectSelected(UObject* InObject)
+{
+	if (DetailsView.IsValid())
+	{
+		DetailsView->SetObject(InObject);
+	}
+}
+
+void FAnimationEditor::PostUndo(bool bSuccess)
+{
+	OnPostUndo.Broadcast();
+}
+
+void FAnimationEditor::PostRedo(bool bSuccess)
+{
+	OnPostUndo.Broadcast();
+}
+
+void FAnimationEditor::HandleDetailsCreated(const TSharedRef<IDetailsView>& InDetailsView)
+{
+	DetailsView = InDetailsView;
+}
+
+TSharedPtr<SDockTab> FAnimationEditor::OpenNewAnimationDocumentTab(UAnimationAsset* InAnimAsset)
+{
+	TSharedPtr<SDockTab> OpenedTab;
+
+	if (InAnimAsset != nullptr)
+	{
+		FString	DocumentLink;
+
+		FAnimDocumentArgs Args(PersonaToolkit->GetPreviewScene(), GetPersonaToolkit(), GetSkeletonTree()->GetEditableSkeleton(), OnPostUndo, OnCurvesChanged, OnChangeAnimNotifies, OnSectionsChanged);
+		Args.OnDespatchObjectsSelected = FOnObjectsSelected::CreateSP(this, &FAnimationEditor::HandleObjectsSelected);
+		Args.OnDespatchAnimNotifiesChanged = FSimpleDelegate::CreateSP(this, &FAnimationEditor::HandleAnimNotifiesChanged);
+		Args.OnDespatchInvokeTab = FOnInvokeTab::CreateSP(this, &FAssetEditorToolkit::InvokeTab);
+		Args.OnDespatchCurvesChanged = FSimpleDelegate::CreateSP(this, &FAnimationEditor::HandleCurvesChanged);
+		Args.OnDespatchSectionsChanged = FSimpleDelegate::CreateSP(this, &FAnimationEditor::HandleSectionsChanged);
+
+		FPersonaModule& PersonaModule = FModuleManager::GetModuleChecked<FPersonaModule>("Persona");
+		TSharedRef<SWidget> TabContents = PersonaModule.CreateEditorWidgetForAnimDocument(SharedThis(this), InAnimAsset, Args, DocumentLink);
+
+		if (AnimationAsset)
+		{
+			RemoveEditingObject(AnimationAsset);
+		}
+
+		if (InAnimAsset != nullptr)
+		{
+			AddEditingObject(InAnimAsset);
+			AnimationAsset = InAnimAsset;
+		}
+
+		GetPersonaToolkit()->GetPreviewScene()->SetPreviewAnimationAsset(InAnimAsset);
+		GetPersonaToolkit()->SetAnimationAsset(InAnimAsset);
+
+		struct Local
+		{
+			static FText GetObjectName(UObject* Object)
+			{
+				return FText::FromString(Object->GetName());
+			}
+		};
+
+		TAttribute<FText> NameAttribute = TAttribute<FText>::Create(TAttribute<FText>::FGetter::CreateStatic(&Local::GetObjectName, (UObject*)InAnimAsset));
+
+		if (SharedAnimDocumentTab.IsValid())
+		{
+			OpenedTab = SharedAnimDocumentTab.Pin();
+			OpenedTab->SetContent(TabContents);
+			OpenedTab->ActivateInParent(ETabActivationCause::SetDirectly);
+			OpenedTab->SetLabel(NameAttribute);
+			OpenedTab->SetLeftContent(IDocumentation::Get()->CreateAnchor(DocumentLink));
+		}
+		else
+		{
+			OpenedTab = SNew(SDockTab)
+				.Label(NameAttribute)
+				.TabRole(ETabRole::DocumentTab)
+				.TabColorScale(GetTabColorScale())
+				[
+					TabContents
+				];
+
+			OpenedTab->SetLeftContent(IDocumentation::Get()->CreateAnchor(DocumentLink));
+
+			TabManager->InsertNewDocumentTab(AnimationEditorTabs::DocumentTab, FTabManager::ESearchPreference::RequireClosedTab, OpenedTab.ToSharedRef());
+
+			SharedAnimDocumentTab = OpenedTab;
+		}
+
+ 		if (SequenceBrowser.IsValid())
+ 		{
+ 			SequenceBrowser.Pin()->SelectAsset(InAnimAsset);
+ 		}
+
+		// let the asset family know too
+		TSharedRef<IAssetFamily> AssetFamily = PersonaModule.CreatePersonaAssetFamily(InAnimAsset);
+		AssetFamily->RecordAssetOpened(FAssetData(InAnimAsset));
+	}
+
+	return OpenedTab;
+}
+
+void FAnimationEditor::HandleAnimNotifiesChanged()
+{
+	OnChangeAnimNotifies.Broadcast();
+}
+
+void FAnimationEditor::HandleCurvesChanged()
+{
+	OnCurvesChanged.Broadcast();
+}
+
+void FAnimationEditor::HandleSectionsChanged()
+{
+	OnSectionsChanged.Broadcast();
+}
+
+void FAnimationEditor::HandleOpenNewAsset(UObject* InNewAsset)
+{
+	if (UAnimationAsset* NewAnimationAsset = Cast<UAnimationAsset>(InNewAsset))
+	{
+		OpenNewAnimationDocumentTab(NewAnimationAsset);
+	}
+}
+
+UObject* FAnimationEditor::HandleGetAsset()
+{
+	return GetEditingObject();
+}
+
+void FAnimationEditor::HandleSetKeyCompleted()
+{
+	OnCurvesChanged.Broadcast();
+}
+
+bool FAnimationEditor::HasValidAnimationSequence() const
+{
+	UAnimSequence* AnimSequence = Cast<UAnimSequence>(AnimationAsset);
+	return (AnimSequence != nullptr);
+}
+
+bool FAnimationEditor::CanSetKey() const
+{
+	UDebugSkelMeshComponent* PreviewMeshComponent = PersonaToolkit->GetPreviewMeshComponent();
+	return (HasValidAnimationSequence() && PreviewMeshComponent->BonesOfInterest.Num() > 0);
+}
+
+void FAnimationEditor::OnSetKey()
+{
+	if (AnimationAsset)
+	{
+		UDebugSkelMeshComponent* Component = PersonaToolkit->GetPreviewMeshComponent();
+		Component->PreviewInstance->SetKey(FSimpleDelegate::CreateSP(this, &FAnimationEditor::HandleSetKeyCompleted));
+	}
+}
+
+bool FAnimationEditor::CanApplyRawAnimChanges() const
+{
+	UAnimSequence* AnimSequence = Cast<UAnimSequence>(AnimationAsset);
+
+	// ideally would be great if we can only show if something changed
+	return (AnimSequence && (AnimSequence->DoesNeedRebake() || AnimSequence->DoesNeedRecompress()));
+}
+
+void FAnimationEditor::OnApplyRawAnimChanges()
+{
+	UAnimSequence* AnimSequence = Cast<UAnimSequence>(AnimationAsset);
+	if (AnimSequence)
+	{
+		if (AnimSequence->DoesNeedRebake() || AnimSequence->DoesNeedRecompress())
+		{
+			FScopedTransaction ScopedTransaction(LOCTEXT("BakeAnimation", "Bake Animation"));
+			if (AnimSequence->DoesNeedRebake())
+			{
+				AnimSequence->Modify(true);
+				AnimSequence->BakeTrackCurvesToRawAnimation();
+			}
+
+			if (AnimSequence->DoesNeedRecompress())
+			{
+				AnimSequence->Modify(true);
+				AnimSequence->RequestSyncAnimRecompression(false);
+			}
+		}
+	}
+}
+
+void FAnimationEditor::OnApplyCompression()
+{
+	UAnimSequence* AnimSequence = Cast<UAnimSequence>(AnimationAsset);
+	if (AnimSequence)
+	{
+		TArray<TWeakObjectPtr<UAnimSequence>> AnimSequences;
+		AnimSequences.Add(AnimSequence);
+		FPersonaModule& PersonaModule = FModuleManager::GetModuleChecked<FPersonaModule>("Persona");
+		PersonaModule.ApplyCompression(AnimSequences);
+	}
+}
+
+void FAnimationEditor::OnExportToFBX()
+{
+	UAnimSequence* AnimSequence = Cast<UAnimSequence>(AnimationAsset);
+	if (AnimSequence)
+	{
+		TArray<TWeakObjectPtr<UAnimSequence>> AnimSequences;
+		AnimSequences.Add(AnimSequence);
+		FPersonaModule& PersonaModule = FModuleManager::GetModuleChecked<FPersonaModule>("Persona");
+		PersonaModule.ExportToFBX(AnimSequences, GetPersonaToolkit()->GetPreviewScene()->GetPreviewMeshComponent()->SkeletalMesh);
+	}
+}
+
+void FAnimationEditor::OnAddLoopingInterpolation()
+{
+	UAnimSequence* AnimSequence = Cast<UAnimSequence>(AnimationAsset);
+	if (AnimSequence)
+	{
+		TArray<TWeakObjectPtr<UAnimSequence>> AnimSequences;
+		AnimSequences.Add(AnimSequence);
+		FPersonaModule& PersonaModule = FModuleManager::GetModuleChecked<FPersonaModule>("Persona");
+		PersonaModule.AddLoopingInterpolation(AnimSequences);
+	}
+}
+
+TSharedRef< SWidget > FAnimationEditor::GenerateCreateAssetMenu() const
+{
+	const bool bShouldCloseWindowAfterMenuSelection = true;
+	FMenuBuilder MenuBuilder(bShouldCloseWindowAfterMenuSelection, NULL);
+
+	// Create Animation menu
+	MenuBuilder.BeginSection("CreateAnimation", LOCTEXT("CreateAnimationMenuHeading", "Animation"));
+	{
+		// create menu
+		MenuBuilder.AddSubMenu(
+			LOCTEXT("CreateAnimationSubmenu", "Create Animation"),
+			LOCTEXT("CreateAnimationSubmenu_ToolTip", "Create Animation for this skeleton"),
+			FNewMenuDelegate::CreateSP(this, &FAnimationEditor::FillCreateAnimationMenu),
+			false,
+			FSlateIcon(FEditorStyle::GetStyleSetName(), "Persona.AssetActions.CreateAnimAsset")
+			);
+
+		MenuBuilder.AddSubMenu(
+			LOCTEXT("CreatePoseAssetSubmenu", "Create PoseAsset"),
+			LOCTEXT("CreatePoseAsssetSubmenu_ToolTip", "Create PoseAsset for this skeleton"),
+			FNewMenuDelegate::CreateSP(this, &FAnimationEditor::FillCreatePoseAssetMenu),
+			false,
+			FSlateIcon(FEditorStyle::GetStyleSetName(), "ClassIcon.PoseAsset")
+		);
+	}
+	MenuBuilder.EndSection();
+
+	TArray<TWeakObjectPtr<USkeleton>> Skeletons;
+
+	Skeletons.Add(PersonaToolkit->GetSkeleton());
+
+	AnimationEditorUtils::FillCreateAssetMenu(MenuBuilder, Skeletons, FAnimAssetCreated::CreateSP(this, &FAnimationEditor::HandleAssetCreated), false);
+
+	return MenuBuilder.MakeWidget();
+}
+
+void FAnimationEditor::FillCreateAnimationMenu(FMenuBuilder& MenuBuilder) const
+{
+	TArray<TWeakObjectPtr<USkeleton>> Skeletons;
+
+	Skeletons.Add(PersonaToolkit->GetSkeleton());
+
+	// create rig
+	MenuBuilder.BeginSection("CreateAnimationSubMenu", LOCTEXT("CreateAnimationSubMenuHeading", "Create Animation"));
+	{
+		MenuBuilder.AddMenuEntry(
+			LOCTEXT("CreateAnimation_RefPose", "From Reference Pose"),
+			LOCTEXT("CreateAnimation_RefPose_Tooltip", "Create Animation from reference pose."),
+			FSlateIcon(),
+			FUIAction(
+				FExecuteAction::CreateStatic(&AnimationEditorUtils::ExecuteNewAnimAsset<UAnimSequenceFactory, UAnimSequence>, Skeletons, FString("_Sequence"), FAnimAssetCreated::CreateSP(this, &FAnimationEditor::CreateAnimation, 0), false),
+				FCanExecuteAction()
+				)
+			);
+
+		MenuBuilder.AddMenuEntry(
+			LOCTEXT("CreateAnimation_CurrentPose", "From Current Pose"),
+			LOCTEXT("CreateAnimation_CurrentPose_Tooltip", "Create Animation from current pose."),
+			FSlateIcon(),
+			FUIAction(
+				FExecuteAction::CreateStatic(&AnimationEditorUtils::ExecuteNewAnimAsset<UAnimSequenceFactory, UAnimSequence>, Skeletons, FString("_Sequence"), FAnimAssetCreated::CreateSP(this, &FAnimationEditor::CreateAnimation, 1), false),
+				FCanExecuteAction()
+				)
+			);
+
+		MenuBuilder.AddMenuEntry(
+			LOCTEXT("CreateAnimation_CurrentAnimation", "From Current Animation"),
+			LOCTEXT("CreateAnimation_CurrentAnimation_Tooltip", "Create Animation from current animation."),
+			FSlateIcon(),
+			FUIAction(
+				FExecuteAction::CreateStatic(&AnimationEditorUtils::ExecuteNewAnimAsset<UAnimSequenceFactory, UAnimSequence>, Skeletons, FString("_Sequence"), FAnimAssetCreated::CreateSP(this, &FAnimationEditor::CreateAnimation, 2), false),
+				FCanExecuteAction::CreateSP(this, &FAnimationEditor::HasValidAnimationSequence)
+				)
+			);
+	}
+	MenuBuilder.EndSection();
+}
+
+void FAnimationEditor::FillCreatePoseAssetMenu(FMenuBuilder& MenuBuilder) const
+{
+	TArray<TWeakObjectPtr<USkeleton>> Skeletons;
+
+	Skeletons.Add(PersonaToolkit->GetSkeleton());
+
+	// create rig
+	MenuBuilder.BeginSection("CreatePoseAssetSubMenu", LOCTEXT("CreatePoseAssetSubMenuHeading", "Create PoseAsset"));
+	{
+		MenuBuilder.AddMenuEntry(
+			LOCTEXT("CreatePoseAsset_CurrentPose", "From Current Pose"),
+			LOCTEXT("CreatePoseAsset_CurrentPose_Tooltip", "Create PoseAsset from current pose."),
+			FSlateIcon(),
+			FUIAction(
+				FExecuteAction::CreateStatic(&AnimationEditorUtils::ExecuteNewAnimAsset<UPoseAssetFactory, UPoseAsset>, Skeletons, FString("_PoseAsset"), FAnimAssetCreated::CreateSP(this, &FAnimationEditor::CreatePoseAsset, 0), false),
+				FCanExecuteAction()
+			)
+		);
+
+		MenuBuilder.AddMenuEntry(
+			LOCTEXT("CreatePoseAsset_CurrentAnimation", "From Current Animation"),
+			LOCTEXT("CreatePoseAsset_CurrentAnimation_Tooltip", "Create Animation from current animation."),
+			FSlateIcon(),
+			FUIAction(
+				FExecuteAction::CreateStatic(&AnimationEditorUtils::ExecuteNewAnimAsset<UPoseAssetFactory, UPoseAsset>, Skeletons, FString("_PoseAsset"), FAnimAssetCreated::CreateSP(this, &FAnimationEditor::CreatePoseAsset, 1), false),
+				FCanExecuteAction()
+			)
+		);
+	}
+	MenuBuilder.EndSection();
+
+	// create pose asset
+	MenuBuilder.BeginSection("InsertPoseSubMenuSection", LOCTEXT("InsertPoseSubMenuSubMenuHeading", "Insert Pose"));
+	{
+		MenuBuilder.AddSubMenu(
+			LOCTEXT("InsertPoseSubmenu", "Insert Pose"),
+			LOCTEXT("InsertPoseSubmenu_ToolTip", "Insert current pose to selected PoseAsset"),
+			FNewMenuDelegate::CreateSP(this, &FAnimationEditor::FillInsertPoseMenu),
+			false,
+			FSlateIcon(FEditorStyle::GetStyleSetName(), "ClassIcon.PoseAsset")
+		);
+	}
+	MenuBuilder.EndSection();
+}
+
+void FAnimationEditor::FillInsertPoseMenu(FMenuBuilder& MenuBuilder) const
+{
+	FAssetPickerConfig AssetPickerConfig;
+
+	USkeleton* Skeleton = GetPersonaToolkit()->GetSkeleton();
+
+	/** The asset picker will only show skeletons */
+	AssetPickerConfig.Filter.ClassNames.Add(*UPoseAsset::StaticClass()->GetName());
+	AssetPickerConfig.Filter.bRecursiveClasses = false;
+	AssetPickerConfig.bAllowNullSelection = false;
+	AssetPickerConfig.Filter.TagsAndValues.Add(TEXT("Skeleton"), FAssetData(Skeleton).GetExportTextName());
+
+	/** The delegate that fires when an asset was selected */
+	AssetPickerConfig.OnAssetSelected = FOnAssetSelected::CreateRaw(this, &FAnimationEditor::InsertCurrentPoseToAsset);
+
+	/** The default view mode should be a list view */
+	AssetPickerConfig.InitialAssetViewType = EAssetViewType::List;
+
+	FContentBrowserModule& ContentBrowserModule = FModuleManager::Get().LoadModuleChecked<FContentBrowserModule>(TEXT("ContentBrowser"));
+
+	MenuBuilder.AddWidget(
+		ContentBrowserModule.Get().CreateAssetPicker(AssetPickerConfig),
+		LOCTEXT("Select_Label", "")
+	);
+}
+
+void FAnimationEditor::InsertCurrentPoseToAsset(const FAssetData& NewPoseAssetData)
+{
+	UPoseAsset* PoseAsset = Cast<UPoseAsset>(NewPoseAssetData.GetAsset());
+	FScopedTransaction ScopedTransaction(LOCTEXT("InsertPose", "Insert Pose"));
+
+	if (PoseAsset)
+	{
+		PoseAsset->Modify();
+
+		UDebugSkelMeshComponent* PreviewMeshComponent = PersonaToolkit->GetPreviewMeshComponent();
+		if (PreviewMeshComponent)
+		{
+			FSmartName NewPoseName;
+
+			bool bSuccess = PoseAsset->AddOrUpdatePoseWithUniqueName(PreviewMeshComponent, &NewPoseName);
+
+			if (bSuccess)
+			{
+				FFormatNamedArguments Args;
+				Args.Add(TEXT("PoseAsset"), FText::FromString(PoseAsset->GetName()));
+				Args.Add(TEXT("PoseName"), FText::FromName(NewPoseName.DisplayName));
+				FNotificationInfo Info(FText::Format(LOCTEXT("InsertPoseSucceeded", "The current pose has inserted to {PoseAsset} with {PoseName}"), Args));
+				Info.ExpireDuration = 7.0f;
+				Info.bUseLargeFont = false;
+				TSharedPtr<SNotificationItem> Notification = FSlateNotificationManager::Get().AddNotification(Info);
+				if (Notification.IsValid())
+				{
+					Notification->SetCompletionState(SNotificationItem::CS_Success);
+				}
+			}
+			else
+			{
+				FFormatNamedArguments Args;
+				Args.Add(TEXT("PoseAsset"), FText::FromString(PoseAsset->GetName()));
+				FNotificationInfo Info(FText::Format(LOCTEXT("InsertPoseFailed", "Inserting pose to asset {PoseAsset} has failed"), Args));
+				Info.ExpireDuration = 7.0f;
+				Info.bUseLargeFont = false;
+				TSharedPtr<SNotificationItem> Notification = FSlateNotificationManager::Get().AddNotification(Info);
+				if (Notification.IsValid())
+				{
+					Notification->SetCompletionState(SNotificationItem::CS_Fail);
+				}
+			}
+		}
+	}
+
+	// it doesn't work well if I leave the window open. The delegate goes weired or it stop showing the popups. 
+	FSlateApplication::Get().DismissAllMenus();
+}
+
+void FAnimationEditor::CreateAnimation(const TArray<UObject*> NewAssets, int32 Option)
+{
+	bool bResult = true;
+	if (NewAssets.Num() > 0)
+	{
+		USkeletalMeshComponent* MeshComponent = PersonaToolkit->GetPreviewMeshComponent();
+		UAnimSequence* Sequence = Cast<UAnimSequence>(AnimationAsset);
+
+		for (auto NewAsset : NewAssets)
+		{
+			UAnimSequence* NewAnimSequence = Cast<UAnimSequence>(NewAsset);
+			if (NewAnimSequence)
+			{
+				switch (Option)
+				{
+				case 0:
+					bResult &= NewAnimSequence->CreateAnimation(MeshComponent->SkeletalMesh);
+					break;
+				case 1:
+					bResult &= NewAnimSequence->CreateAnimation(MeshComponent);
+					break;
+				case 2:
+					bResult &= NewAnimSequence->CreateAnimation(Sequence);
+					break;
+				}
+			}
+		}
+
+		if (bResult)
+		{
+			HandleAssetCreated(NewAssets);
+
+			// if it created based on current mesh component, 
+			if (Option == 1)
+			{
+				UDebugSkelMeshComponent* PreviewMeshComponent = PersonaToolkit->GetPreviewMeshComponent();
+				if (PreviewMeshComponent && PreviewMeshComponent->PreviewInstance)
+				{
+					PreviewMeshComponent->PreviewInstance->ResetModifiedBone();
+				}
+			}
+		}
+	}
+}
+
+void FAnimationEditor::CreatePoseAsset(const TArray<UObject*> NewAssets, int32 Option)
+{
+	bool bResult = false;
+	if (NewAssets.Num() > 0)
+	{
+		UDebugSkelMeshComponent* PreviewComponent = PersonaToolkit->GetPreviewMeshComponent();
+		UAnimSequence* Sequence = Cast<UAnimSequence>(AnimationAsset);
+
+		for (auto NewAsset : NewAssets)
+		{
+			UPoseAsset* NewPoseAsset = Cast<UPoseAsset>(NewAsset);
+			if (NewPoseAsset)
+			{
+				switch (Option)
+				{
+				case 0:
+					NewPoseAsset->AddOrUpdatePoseWithUniqueName(PreviewComponent);
+					break;
+				case 1:
+					NewPoseAsset->CreatePoseFromAnimation(Sequence);
+					break;
+				}
+
+				bResult = true;
+			}
+		}
+
+		// if it contains error, warn them
+		if (bResult)
+		{
+			HandleAssetCreated(NewAssets);
+
+			// if it created based on current mesh component, 
+			if (Option == 0)
+			{
+				PreviewComponent->PreviewInstance->ResetModifiedBone();
+			}
+		}
+		else
+		{
+			FMessageDialog::Open(EAppMsgType::Ok, LOCTEXT("FailedToCreateAsset", "Failed to create asset"));
+		}
+	}
+}
+
+void FAnimationEditor::HandleAssetCreated(const TArray<UObject*> NewAssets)
+{
+	if (NewAssets.Num() > 0)
+	{
+		FAssetRegistryModule::AssetCreated(NewAssets[0]);
+		OpenNewAnimationDocumentTab(CastChecked<UAnimationAsset>(NewAssets[0]));
+	}
+}
+
+void FAnimationEditor::ConditionalRefreshEditor(UObject* InObject)
+{
+	bool bInterestingAsset = true;
+
+	if (InObject != GetPersonaToolkit()->GetSkeleton() && InObject != GetPersonaToolkit()->GetSkeleton()->GetPreviewMesh() && Cast<UAnimationAsset>(InObject) != AnimationAsset)
+	{
+		bInterestingAsset = false;
+	}
+
+	// Check that we aren't a montage that uses an incoming animation
+	if (UAnimMontage* Montage = Cast<UAnimMontage>(AnimationAsset))
+	{
+		for (FSlotAnimationTrack& Slot : Montage->SlotAnimTracks)
+		{
+			if (bInterestingAsset)
+			{
+				break;
+			}
+
+			for (FAnimSegment& Segment : Slot.AnimTrack.AnimSegments)
+			{
+				if (Segment.AnimReference == InObject)
+				{
+					bInterestingAsset = true;
+					break;
+				}
+			}
+		}
+	}
+
+	if (bInterestingAsset)
+	{
+		GetPersonaToolkit()->GetPreviewScene()->InvalidateViews();
+		OpenNewAnimationDocumentTab(CastChecked<UAnimationAsset>(InObject));
+	}
+}
+
+void FAnimationEditor::HandlePostReimport(UObject* InObject, bool bSuccess)
+{
+	if (bSuccess)
+	{
+		ConditionalRefreshEditor(InObject);
+	}
+}
+
+void FAnimationEditor::HandlePostImport(UFactory* InFactory, UObject* InObject)
+{
+	ConditionalRefreshEditor(InObject);
+}
+
+void FAnimationEditor::HandleAnimationSequenceBrowserCreated(const TSharedRef<IAnimationSequenceBrowser>& InSequenceBrowser)
+{
+	SequenceBrowser = InSequenceBrowser;
+}
+
+#undef LOCTEXT_NAMESPACE
diff --git a/Engine/Source/Editor/AnimationEditor/Private/AnimationEditorCommands.h b/Engine/Source/Editor/AnimationEditor/Private/AnimationEditorCommands.h
new file mode 100644
--- /dev/null
+++ b/Engine/Source/Editor/AnimationEditor/Private/AnimationEditorCommands.h
@@ -0,0 +1,30 @@
+// Copyright 1998-2016 Epic Games, Inc. All Rights Reserved.
+
+#pragma once
+
+class FAnimationEditorCommands : public TCommands<FAnimationEditorCommands>
+{
+public:
+	FAnimationEditorCommands()
+		: TCommands<FAnimationEditorCommands>(TEXT("AnimationEditor"), NSLOCTEXT("Contexts", "AnimationEditor", "Animation Editor"), NAME_None, FEditorStyle::GetStyleSetName())
+	{
+	}
+
+	virtual void RegisterCommands() override;
+
+public:
+	// import animation
+	TSharedPtr<FUICommandInfo> ImportAnimation;
+	// reimport animation
+	TSharedPtr<FUICommandInfo> ReimportAnimation;
+	// apply compression
+	TSharedPtr<FUICommandInfo> ApplyCompression;
+	// export to FBX
+	TSharedPtr<FUICommandInfo> ExportToFBX;
+	// Add looping interpolation
+	TSharedPtr<FUICommandInfo> AddLoopingInterpolation;
+	// set key for bone track
+	TSharedPtr<FUICommandInfo> SetKey;
+	// bake bone track curves to animation
+	TSharedPtr<FUICommandInfo> ApplyAnimation;
+};
diff --git a/Engine/Source/Editor/AnimationEditor/Private/AnimationEditorMode.cpp b/Engine/Source/Editor/AnimationEditor/Private/AnimationEditorMode.cpp
new file mode 100644
--- /dev/null
+++ b/Engine/Source/Editor/AnimationEditor/Private/AnimationEditorMode.cpp
@@ -0,0 +1,112 @@
+// Copyright 1998-2016 Epic Games, Inc. All Rights Reserved.
+
+#include "AnimationEditorPrivatePCH.h"
+#include "AnimationEditorMode.h"
+#include "AnimationEditor.h"
+#include "ISkeletonEditorModule.h"
+#include "IPersonaToolkit.h"
+
+FAnimationEditorMode::FAnimationEditorMode(TSharedRef<FWorkflowCentricApplication> InHostingApp, TSharedRef<ISkeletonTree> InSkeletonTree)
+	: FApplicationMode(AnimationEditorModes::AnimationEditorMode)
+{
+	HostingAppPtr = InHostingApp;
+
+	TSharedRef<FAnimationEditor> AnimationEditor = StaticCastSharedRef<FAnimationEditor>(InHostingApp);
+
+	ISkeletonEditorModule& SkeletonEditorModule = FModuleManager::LoadModuleChecked<ISkeletonEditorModule>("SkeletonEditor");
+	TabFactories.RegisterFactory(SkeletonEditorModule.CreateSkeletonTreeTabFactory(InHostingApp, InSkeletonTree));
+
+	FPersonaModule& PersonaModule = FModuleManager::LoadModuleChecked<FPersonaModule>("Persona");
+	TabFactories.RegisterFactory(PersonaModule.CreateDetailsTabFactory(InHostingApp, FOnDetailsCreated::CreateSP(&AnimationEditor.Get(), &FAnimationEditor::HandleDetailsCreated)));
+	TabFactories.RegisterFactory(PersonaModule.CreatePersonaViewportTabFactory(InHostingApp, InSkeletonTree, AnimationEditor->GetPersonaToolkit()->GetPreviewScene(), AnimationEditor->OnPostUndo, nullptr, FOnViewportCreated(), false, true));
+	TabFactories.RegisterFactory(PersonaModule.CreateAdvancedPreviewSceneTabFactory(InHostingApp, AnimationEditor->GetPersonaToolkit()->GetPreviewScene()));
+	TabFactories.RegisterFactory(PersonaModule.CreateAnimationAssetBrowserTabFactory(InHostingApp, AnimationEditor->GetPersonaToolkit(), FOnOpenNewAsset::CreateSP(&AnimationEditor.Get(), &FAnimationEditor::HandleOpenNewAsset), FOnAnimationSequenceBrowserCreated::CreateSP(&AnimationEditor.Get(), &FAnimationEditor::HandleAnimationSequenceBrowserCreated), true));
+	TabFactories.RegisterFactory(PersonaModule.CreateAssetDetailsTabFactory(InHostingApp, FOnGetAsset::CreateSP(&AnimationEditor.Get(), &FAnimationEditor::HandleGetAsset), FOnDetailsCreated()));
+	TabFactories.RegisterFactory(PersonaModule.CreateCurveViewerTabFactory(InHostingApp, InSkeletonTree->GetEditableSkeleton(), AnimationEditor->GetPersonaToolkit()->GetPreviewScene(), AnimationEditor->OnCurvesChanged, AnimationEditor->OnPostUndo));
+	TabFactories.RegisterFactory(PersonaModule.CreateSkeletonSlotNamesTabFactory(InHostingApp, InSkeletonTree->GetEditableSkeleton(), AnimationEditor->OnPostUndo, FOnObjectSelected::CreateSP(&AnimationEditor.Get(), &FAnimationEditor::HandleObjectSelected)));
+
+	TabLayout = FTabManager::NewLayout("Standalone_AnimationEditor_Layout_v1.1")
+		->AddArea
+		(
+			FTabManager::NewPrimaryArea()
+			->SetOrientation(Orient_Vertical)
+			->Split
+			(
+				FTabManager::NewStack()
+				->SetSizeCoefficient(0.1f)
+				->SetHideTabWell(true)
+				->AddTab(InHostingApp->GetToolbarTabId(), ETabState::OpenedTab)
+			)
+			->Split
+			(
+				FTabManager::NewSplitter()
+				->SetSizeCoefficient(0.9f)
+				->SetOrientation(Orient_Horizontal)
+				->Split
+				(
+					FTabManager::NewSplitter()
+					->SetSizeCoefficient(0.2f)
+					->SetOrientation(Orient_Vertical)
+					->Split
+					(
+						FTabManager::NewStack()
+						->SetHideTabWell(false)
+						->AddTab(AnimationEditorTabs::SkeletonTreeTab, ETabState::OpenedTab)
+						->AddTab(AnimationEditorTabs::AssetDetailsTab, ETabState::OpenedTab)
+					)
+				)
+				->Split
+				(
+					FTabManager::NewSplitter()
+					->SetSizeCoefficient(0.6f)
+					->SetOrientation(Orient_Vertical)
+					->Split
+					(
+						FTabManager::NewStack()
+						->SetSizeCoefficient(0.6f)
+						->SetHideTabWell(true)
+						->AddTab(AnimationEditorTabs::ViewportTab, ETabState::OpenedTab)
+					)
+					->Split
+					(
+						FTabManager::NewStack()
+						->SetSizeCoefficient(0.4f)
+						->SetHideTabWell(true)
+						->AddTab(AnimationEditorTabs::DocumentTab, ETabState::ClosedTab)
+					)
+				)
+				->Split
+				(
+					FTabManager::NewSplitter()
+					->SetSizeCoefficient(0.2f)
+					->SetOrientation(Orient_Vertical)
+					->Split
+					(
+						FTabManager::NewStack()
+						->SetSizeCoefficient(0.6f)
+						->SetHideTabWell(false)
+						->AddTab(AnimationEditorTabs::DetailsTab, ETabState::OpenedTab)
+						->AddTab(AnimationEditorTabs::AdvancedPreviewTab, ETabState::ClosedTab)
+					)
+					->Split
+					(
+						FTabManager::NewStack()
+						->SetSizeCoefficient(0.4f)
+						->SetHideTabWell(false)
+						->AddTab(AnimationEditorTabs::AssetBrowserTab, ETabState::OpenedTab)
+						->AddTab(AnimationEditorTabs::CurveNamesTab, ETabState::ClosedTab)
+						->AddTab(AnimationEditorTabs::SlotNamesTab, ETabState::ClosedTab)
+					)
+				)
+			)
+		);
+}
+
+void FAnimationEditorMode::RegisterTabFactories(TSharedPtr<FTabManager> InTabManager)
+{
+	TSharedPtr<FWorkflowCentricApplication> HostingApp = HostingAppPtr.Pin();
+	HostingApp->RegisterTabSpawners(InTabManager.ToSharedRef());
+	HostingApp->PushTabFactories(TabFactories);
+
+	FApplicationMode::RegisterTabFactories(InTabManager);
+}
\ No newline at end of file
diff --git a/Engine/Source/Editor/AnimationEditor/Private/AnimationEditorMode.h b/Engine/Source/Editor/AnimationEditor/Private/AnimationEditorMode.h
new file mode 100644
--- /dev/null
+++ b/Engine/Source/Editor/AnimationEditor/Private/AnimationEditorMode.h
@@ -0,0 +1,23 @@
+// Copyright 1998-2016 Epic Games, Inc. All Rights Reserved.
+
+#pragma once
+
+#include "ApplicationMode.h"
+#include "WorkflowCentricApplication.h"
+#include "PersonaModule.h"
+
+class FAnimationEditorMode : public FApplicationMode
+{
+public:
+	FAnimationEditorMode(TSharedRef<class FWorkflowCentricApplication> InHostingApp, TSharedRef<class ISkeletonTree> InSkeletonTree);
+
+	/** FApplicationMode interface */
+	virtual void RegisterTabFactories(TSharedPtr<FTabManager> InTabManager) override;
+
+protected:
+	/** The hosting app */
+	TWeakPtr<class FWorkflowCentricApplication> HostingAppPtr;
+
+	/** The tab factories we support */
+	FWorkflowAllowedTabSet TabFactories;
+};
\ No newline at end of file
diff --git a/Engine/Source/Editor/AnimationEditor/Private/AnimationEditorPrivatePCH.h b/Engine/Source/Editor/AnimationEditor/Private/AnimationEditorPrivatePCH.h
new file mode 100644
--- /dev/null
+++ b/Engine/Source/Editor/AnimationEditor/Private/AnimationEditorPrivatePCH.h
@@ -0,0 +1,8 @@
+// Copyright 1998-2016 Epic Games, Inc. All Rights Reserved.
+
+#pragma once
+
+#include "UnrealEd.h"
+#include "SlateBasics.h"
+#include "Toolkits/AssetEditorToolkit.h"
+#include "Toolkits/IToolkit.h"
\ No newline at end of file
diff --git a/Engine/Source/Editor/AnimationEditor/Public/IAnimationEditor.h b/Engine/Source/Editor/AnimationEditor/Public/IAnimationEditor.h
new file mode 100644
--- /dev/null
+++ b/Engine/Source/Editor/AnimationEditor/Public/IAnimationEditor.h
@@ -0,0 +1,10 @@
+// Copyright 1998-2016 Epic Games, Inc. All Rights Reserved.
+
+#pragma once
+
+#include "WorkflowCentricApplication.h"
+#include "IHasPersonaToolkit.h"
+
+class IAnimationEditor : public FWorkflowCentricApplication, public IHasPersonaToolkit
+{
+};
diff --git a/Engine/Source/Editor/BlueprintGraph/Private/K2Node_Variable.cpp b/Engine/Source/Editor/BlueprintGraph/Private/K2Node_Variable.cpp
--- a/Engine/Source/Editor/BlueprintGraph/Private/K2Node_Variable.cpp
+++ b/Engine/Source/Editor/BlueprintGraph/Private/K2Node_Variable.cpp
@@ -121,6 +121,19 @@ void UK2Node_Variable::CreatePinForSelf()
 			UClass* MemberParentClass = VariableReference.GetMemberParentClass(GetBlueprintClassFromNode());
 			UClass* TargetClass = MemberParentClass;
 			
+			UProperty* VariableProperty = GetPropertyForVariable();
+
+			if (VariableProperty)
+			{
+				UClass* PropertyClass = CastChecked<UClass>(VariableProperty->GetOuter());
+
+				// Fix up target class if it's not correct, this fixes cases where variables have moved within the hierarchy
+				if (PropertyClass != TargetClass)
+				{
+					TargetClass = PropertyClass;
+				}
+			}
+
 			// Self Target pins should always make the class be the owning class of the property,
 			// so if the node is from a Macro Blueprint, it will hook up as self in any placed Blueprint
 			if(bSelfTarget)
@@ -134,9 +147,9 @@ void UK2Node_Variable::CreatePinForSelf()
 					TargetClass = GetBlueprint()->SkeletonGeneratedClass->GetAuthoritativeClass();
 				}
 			}
-			else if(MemberParentClass && MemberParentClass->ClassGeneratedBy)
+			else if(TargetClass && TargetClass->ClassGeneratedBy)
 			{
-				TargetClass = MemberParentClass->GetAuthoritativeClass();
+				TargetClass = TargetClass->GetAuthoritativeClass();
 			}
 
 			UEdGraphPin* TargetPin = CreatePin(EGPD_Input, K2Schema->PC_Object, TEXT(""), TargetClass, false, false, K2Schema->PN_Self);
