Moved asset creation tools to common Persona toolbar items

This means animations and pose assets can be created from (e.g.) skeletal meshes.

diff --git a/Engine/Source/Editor/AnimationEditor/Private/AnimationEditor.cpp b/Engine/Source/Editor/AnimationEditor/Private/AnimationEditor.cpp
--- a/Engine/Source/Editor/AnimationEditor/Private/AnimationEditor.cpp
+++ b/Engine/Source/Editor/AnimationEditor/Private/AnimationEditor.cpp
@@ -198,11 +198,11 @@ void FAnimationEditor::BindCommands()
 		FCanExecuteAction::CreateSP(this, &FAnimationEditor::CanApplyRawAnimChanges));
 
 	ToolkitCommands->MapAction(FAnimationEditorCommands::Get().ExportToFBX_AnimData,
-		FExecuteAction::CreateSP(this, &FAnimationEditor::OnExportToFBX, EPoseSourceOption::CurrentAnimation_AnimData),
+		FExecuteAction::CreateSP(this, &FAnimationEditor::OnExportToFBX, EExportSourceOption::CurrentAnimation_AnimData),
 		FCanExecuteAction::CreateSP(this, &FAnimationEditor::HasValidAnimationSequence));
 
 	ToolkitCommands->MapAction(FAnimationEditorCommands::Get().ExportToFBX_PreviewMesh,
-		FExecuteAction::CreateSP(this, &FAnimationEditor::OnExportToFBX, EPoseSourceOption::CurrentAnimation_PreviewMesh),
+		FExecuteAction::CreateSP(this, &FAnimationEditor::OnExportToFBX, EExportSourceOption::CurrentAnimation_PreviewMesh),
 		FCanExecuteAction::CreateSP(this, &FAnimationEditor::HasValidAnimationSequence));
 
 	ToolkitCommands->MapAction(FAnimationEditorCommands::Get().AddLoopingInterpolation,
@@ -258,17 +258,6 @@ void FAnimationEditor::ExtendToolbar()
 
 			ToolbarBuilder.BeginSection("Animation");
 			{
-				// create button
-				{
-					ToolbarBuilder.AddComboButton(
-						FUIAction(),
-						FOnGetContent::CreateSP(this, &FAnimationEditor::GenerateCreateAssetMenu),
-						LOCTEXT("CreateAsset_Label", "Create Asset"),
-						LOCTEXT("CreateAsset_ToolTip", "Create Assets for this skeleton."),
-						FSlateIcon(FEditorStyle::GetStyleSetName(), "Persona.CreateAsset")
-					);
-				}
-
 				ToolbarBuilder.AddToolBarButton(FAnimationEditorCommands::Get().ReimportAnimation);
 				ToolbarBuilder.AddToolBarButton(FAnimationEditorCommands::Get().ApplyCompression, NAME_None, LOCTEXT("Toolbar_ApplyCompression", "Compression"));
 
@@ -557,16 +546,16 @@ void FAnimationEditor::OnApplyCompression()
 	}
 }
 
-void FAnimationEditor::OnExportToFBX(const EPoseSourceOption Option)
+void FAnimationEditor::OnExportToFBX(const EExportSourceOption Option)
 {
 	UAnimSequence* AnimSequenceToRecord = nullptr;
-	if (Option == EPoseSourceOption::CurrentAnimation_AnimData)
+	if (Option == EExportSourceOption::CurrentAnimation_AnimData)
 	{
 		TArray<UObject*> AssetsToExport;
 		AssetsToExport.Add(AnimationAsset);
 		ExportToFBX(AssetsToExport, false);
 	}
-	else if (Option == EPoseSourceOption::CurrentAnimation_PreviewMesh)
+	else if (Option == EExportSourceOption::CurrentAnimation_PreviewMesh)
 	{
 		TArray<TWeakObjectPtr<UObject>> Skeletons;
 		Skeletons.Add(PersonaToolkit->GetSkeleton());
@@ -642,266 +631,6 @@ TSharedRef< SWidget > FAnimationEditor::GenerateExportAssetMenu() const
 	return MenuBuilder.MakeWidget();
 }
 
-TSharedRef< SWidget > FAnimationEditor::GenerateCreateAssetMenu() const
-{
-	const bool bShouldCloseWindowAfterMenuSelection = true;
-	FMenuBuilder MenuBuilder(bShouldCloseWindowAfterMenuSelection, NULL);
-
-	// Create Animation menu
-	MenuBuilder.BeginSection("CreateAnimation", LOCTEXT("CreateAnimationMenuHeading", "Animation"));
-	{
-		// create menu
-		MenuBuilder.AddSubMenu(
-			LOCTEXT("CreateAnimationSubmenu", "Create Animation"),
-			LOCTEXT("CreateAnimationSubmenu_ToolTip", "Create Animation for this skeleton"),
-			FNewMenuDelegate::CreateSP(this, &FAnimationEditor::FillCreateAnimationMenu),
-			false,
-			FSlateIcon(FEditorStyle::GetStyleSetName(), "Persona.AssetActions.CreateAnimAsset")
-			);
-
-		MenuBuilder.AddSubMenu(
-			LOCTEXT("CreatePoseAssetSubmenu", "Create PoseAsset"),
-			LOCTEXT("CreatePoseAsssetSubmenu_ToolTip", "Create PoseAsset for this skeleton"),
-			FNewMenuDelegate::CreateSP(this, &FAnimationEditor::FillCreatePoseAssetMenu),
-			false,
-			FSlateIcon(FEditorStyle::GetStyleSetName(), "ClassIcon.PoseAsset")
-		);
-	}
-	MenuBuilder.EndSection();
-
-	TArray<TWeakObjectPtr<UObject>> Objects;
-
-	if (PersonaToolkit->GetPreviewMesh())
-	{
-		Objects.Add(PersonaToolkit->GetPreviewMesh());
-	}
-	else
-	{
-		Objects.Add(PersonaToolkit->GetSkeleton());
-	}
-
-	AnimationEditorUtils::FillCreateAssetMenu(MenuBuilder, Objects, FAnimAssetCreated::CreateSP(this, &FAnimationEditor::HandleAssetCreated), false);
-
-	return MenuBuilder.MakeWidget();
-}
-
-void FAnimationEditor::FillCreateAnimationMenu(FMenuBuilder& MenuBuilder) const
-{
-	TArray<TWeakObjectPtr<UObject>> Objects;
-
-	if (PersonaToolkit->GetPreviewMesh())
-	{
-		Objects.Add(PersonaToolkit->GetPreviewMesh());
-	}
-	else
-	{
-		Objects.Add(PersonaToolkit->GetSkeleton());
-	}
-
-	// create rig
-	MenuBuilder.BeginSection("CreateAnimationSubMenu", LOCTEXT("CreateAnimationSubMenuHeading", "Create Animation"));
-	{
-		MenuBuilder.AddMenuEntry(
-			LOCTEXT("CreateAnimation_RefPose", "Reference Pose"),
-			LOCTEXT("CreateAnimation_RefPose_Tooltip", "Create Animation from reference pose."),
-			FSlateIcon(),
-			FUIAction(
-				FExecuteAction::CreateStatic(&AnimationEditorUtils::ExecuteNewAnimAsset<UAnimSequenceFactory, UAnimSequence>, Objects, FString("_Sequence"), FAnimAssetCreated::CreateSP(this, &FAnimationEditor::CreateAnimation, EPoseSourceOption::ReferencePose), false),
-				FCanExecuteAction()
-				)
-			);
-
-		MenuBuilder.AddMenuEntry(
-			LOCTEXT("CreateAnimation_CurrentPose", "Current Pose"),
-			LOCTEXT("CreateAnimation_CurrentPose_Tooltip", "Create Animation from current pose."),
-			FSlateIcon(),
-			FUIAction(
-				FExecuteAction::CreateStatic(&AnimationEditorUtils::ExecuteNewAnimAsset<UAnimSequenceFactory, UAnimSequence>, Objects, FString("_Sequence"), FAnimAssetCreated::CreateSP(this, &FAnimationEditor::CreateAnimation, EPoseSourceOption::CurrentPose), false),
-				FCanExecuteAction()
-				)
-			);
-
-		MenuBuilder.AddSubMenu(
-			LOCTEXT("CreateAnimation_CurrenAnimationSubMenu", "Current Animation"),
-			LOCTEXT("CreateAnimation_CurrenAnimationSubMenu_ToolTip", "Create Animation from current animation"),
-			FNewMenuDelegate::CreateSP(this, &FAnimationEditor::FillCreateAnimationFromCurrentAnimationMenu),
-			false,
-			FSlateIcon(FEditorStyle::GetStyleSetName(), "Persona.AssetActions.CreateAnimAsset")
-		);
-	}
-	MenuBuilder.EndSection();
-}
-
-void FAnimationEditor::FillCreateAnimationFromCurrentAnimationMenu(FMenuBuilder& MenuBuilder) const
-{
-	{
-		TArray<TWeakObjectPtr<UObject>> Objects;
-
-		if (PersonaToolkit->GetPreviewMesh())
-		{
-			Objects.Add(PersonaToolkit->GetPreviewMesh());
-		}
-		else
-		{
-			Objects.Add(PersonaToolkit->GetSkeleton());
-		}
-
-		// create rig
-		MenuBuilder.BeginSection("CreateAnimationSubMenu", LOCTEXT("CreateAnimationFromCurrentAnimationSubmenuHeading", "Create Animation"));
-		{
-			MenuBuilder.AddMenuEntry(
-				LOCTEXT("CreateAnimation_CurrentAnimation_AnimData", "Animation Data"),
-				LOCTEXT("CreateAnimation_CurrentAnimation_AnimData_Tooltip", "Create Animation from Animation Source Data."),
-				FSlateIcon(),
-				FUIAction(
-					FExecuteAction::CreateStatic(&AnimationEditorUtils::ExecuteNewAnimAsset<UAnimSequenceFactory, UAnimSequence>, Objects, FString("_Sequence"), FAnimAssetCreated::CreateSP(this, &FAnimationEditor::CreateAnimation, EPoseSourceOption::CurrentAnimation_AnimData), false),
-					FCanExecuteAction::CreateSP(this, &FAnimationEditor::HasValidAnimationSequence)
-				)
-			);
-
-			MenuBuilder.AddMenuEntry(
-				LOCTEXT("CreateAnimation_CurrentAnimation_PreviewMesh", "Preview Mesh"),
-				LOCTEXT("CreateAnimation_CurrentAnimation_PreviewMesh_Tooltip", "Create Animation by playing on the Current Preview Mesh, including Retargeting, Post Process Graph, or anything you see on the preview mesh."),
-				FSlateIcon(),
-				FUIAction(
-					FExecuteAction::CreateStatic(&AnimationEditorUtils::ExecuteNewAnimAsset<UAnimSequenceFactory, UAnimSequence>, Objects, FString("_Sequence"), FAnimAssetCreated::CreateSP(this, &FAnimationEditor::CreateAnimation, EPoseSourceOption::CurrentAnimation_PreviewMesh), false),
-					FCanExecuteAction::CreateSP(this, &FAnimationEditor::HasValidAnimationSequence)
-				)
-			);
-		}
-		MenuBuilder.EndSection();
-	}
-
-}
-void FAnimationEditor::FillCreatePoseAssetMenu(FMenuBuilder& MenuBuilder) const
-{
-	TArray<TWeakObjectPtr<UObject>> Objects;
-
-	if (PersonaToolkit->GetPreviewMesh())
-	{
-		Objects.Add(PersonaToolkit->GetPreviewMesh());
-	}
-	else
-	{
-		Objects.Add(PersonaToolkit->GetSkeleton());
-	}
-
-	// create rig
-	MenuBuilder.BeginSection("CreatePoseAssetSubMenu", LOCTEXT("CreatePoseAssetSubMenuHeading", "Create PoseAsset"));
-	{
-		MenuBuilder.AddMenuEntry(
-			LOCTEXT("CreatePoseAsset_CurrentPose", "Current Pose"),
-			LOCTEXT("CreatePoseAsset_CurrentPose_Tooltip", "Create PoseAsset from current pose."),
-			FSlateIcon(),
-			FUIAction(
-				FExecuteAction::CreateStatic(&AnimationEditorUtils::ExecuteNewAnimAsset<UPoseAssetFactory, UPoseAsset>, Objects, FString("_PoseAsset"), FAnimAssetCreated::CreateSP(this, &FAnimationEditor::CreatePoseAsset, EPoseSourceOption::CurrentPose), false),
-				FCanExecuteAction()
-			)
-		);
-
-		MenuBuilder.AddMenuEntry(
-			LOCTEXT("CreatePoseAsset_CurrentAnimation", "Current Animation"),
-			LOCTEXT("CreatePoseAsset_CurrentAnimation_Tooltip", "Create Animation from current animation."),
-			FSlateIcon(),
-			FUIAction(
-				FExecuteAction::CreateStatic(&AnimationEditorUtils::ExecuteNewAnimAsset<UPoseAssetFactory, UPoseAsset>, Objects, FString("_PoseAsset"), FAnimAssetCreated::CreateSP(this, &FAnimationEditor::CreatePoseAsset, EPoseSourceOption::CurrentAnimation_AnimData), false),
-				FCanExecuteAction()
-			)
-		);
-	}
-	MenuBuilder.EndSection();
-
-	// create pose asset
-	MenuBuilder.BeginSection("InsertPoseSubMenuSection", LOCTEXT("InsertPoseSubMenuSubMenuHeading", "Insert Pose"));
-	{
-		MenuBuilder.AddSubMenu(
-			LOCTEXT("InsertPoseSubmenu", "Insert Pose"),
-			LOCTEXT("InsertPoseSubmenu_ToolTip", "Insert current pose to selected PoseAsset"),
-			FNewMenuDelegate::CreateSP(this, &FAnimationEditor::FillInsertPoseMenu),
-			false,
-			FSlateIcon(FEditorStyle::GetStyleSetName(), "ClassIcon.PoseAsset")
-		);
-	}
-	MenuBuilder.EndSection();
-}
-
-void FAnimationEditor::FillInsertPoseMenu(FMenuBuilder& MenuBuilder) const
-{
-	FAssetPickerConfig AssetPickerConfig;
-
-	USkeleton* Skeleton = GetPersonaToolkit()->GetSkeleton();
-
-	/** The asset picker will only show skeletons */
-	AssetPickerConfig.Filter.ClassNames.Add(*UPoseAsset::StaticClass()->GetName());
-	AssetPickerConfig.Filter.bRecursiveClasses = false;
-	AssetPickerConfig.bAllowNullSelection = false;
-	AssetPickerConfig.Filter.TagsAndValues.Add(TEXT("Skeleton"), FAssetData(Skeleton).GetExportTextName());
-
-	/** The delegate that fires when an asset was selected */
-	AssetPickerConfig.OnAssetSelected = FOnAssetSelected::CreateRaw(this, &FAnimationEditor::InsertCurrentPoseToAsset);
-
-	/** The default view mode should be a list view */
-	AssetPickerConfig.InitialAssetViewType = EAssetViewType::List;
-
-	FContentBrowserModule& ContentBrowserModule = FModuleManager::Get().LoadModuleChecked<FContentBrowserModule>(TEXT("ContentBrowser"));
-
-	MenuBuilder.AddWidget(
-		ContentBrowserModule.Get().CreateAssetPicker(AssetPickerConfig),
-		FText::GetEmpty()
-	);
-}
-
-void FAnimationEditor::InsertCurrentPoseToAsset(const FAssetData& NewPoseAssetData)
-{
-	UPoseAsset* PoseAsset = Cast<UPoseAsset>(NewPoseAssetData.GetAsset());
-	FScopedTransaction ScopedTransaction(LOCTEXT("InsertPose", "Insert Pose"));
-
-	if (PoseAsset)
-	{
-		PoseAsset->Modify();
-
-		UDebugSkelMeshComponent* PreviewMeshComponent = PersonaToolkit->GetPreviewMeshComponent();
-		if (PreviewMeshComponent)
-		{
-			FSmartName NewPoseName;
-
-			bool bSuccess = PoseAsset->AddOrUpdatePoseWithUniqueName(PreviewMeshComponent, &NewPoseName);
-
-			if (bSuccess)
-			{
-				FFormatNamedArguments Args;
-				Args.Add(TEXT("PoseAsset"), FText::FromString(PoseAsset->GetName()));
-				Args.Add(TEXT("PoseName"), FText::FromName(NewPoseName.DisplayName));
-				FNotificationInfo Info(FText::Format(LOCTEXT("InsertPoseSucceeded", "The current pose has inserted to {PoseAsset} with {PoseName}"), Args));
-				Info.ExpireDuration = 7.0f;
-				Info.bUseLargeFont = false;
-				TSharedPtr<SNotificationItem> Notification = FSlateNotificationManager::Get().AddNotification(Info);
-				if (Notification.IsValid())
-				{
-					Notification->SetCompletionState(SNotificationItem::CS_Success);
-				}
-			}
-			else
-			{
-				FFormatNamedArguments Args;
-				Args.Add(TEXT("PoseAsset"), FText::FromString(PoseAsset->GetName()));
-				FNotificationInfo Info(FText::Format(LOCTEXT("InsertPoseFailed", "Inserting pose to asset {PoseAsset} has failed"), Args));
-				Info.ExpireDuration = 7.0f;
-				Info.bUseLargeFont = false;
-				TSharedPtr<SNotificationItem> Notification = FSlateNotificationManager::Get().AddNotification(Info);
-				if (Notification.IsValid())
-				{
-					Notification->SetCompletionState(SNotificationItem::CS_Fail);
-				}
-			}
-		}
-	}
-
-	// it doesn't work well if I leave the window open. The delegate goes weired or it stop showing the popups. 
-	FSlateApplication::Get().DismissAllMenus();
-}
-
-
 void FAnimationEditor::FillCopyToSoundWaveMenu(FMenuBuilder& MenuBuilder) const
 {
 	FAssetPickerConfig AssetPickerConfig;
@@ -997,134 +726,6 @@ void FAnimationEditor::CopyCurveToSoundWave(const FAssetData& SoundWaveAssetData
 	FSlateApplication::Get().DismissAllMenus();
 }
 
-bool FAnimationEditor::CreateAnimation(const TArray<UObject*> NewAssets, const EPoseSourceOption Option)
-{
-	bool bResult = true;
-	if (NewAssets.Num() > 0)
-	{
-		USkeletalMeshComponent* MeshComponent = PersonaToolkit->GetPreviewMeshComponent();
-		UAnimSequence* Sequence = Cast<UAnimSequence>(AnimationAsset);
-
-		for (auto NewAsset : NewAssets)
-		{
-			UAnimSequence* NewAnimSequence = Cast<UAnimSequence>(NewAsset);
-			if (NewAnimSequence)
-			{
-				switch (Option)
-				{
-				case EPoseSourceOption::ReferencePose:
-					bResult &= NewAnimSequence->CreateAnimation(MeshComponent->SkeletalMesh);
-					break;
-				case EPoseSourceOption::CurrentPose:
-					bResult &= NewAnimSequence->CreateAnimation(MeshComponent);
-					break;
-				case EPoseSourceOption::CurrentAnimation_AnimData:
-					bResult &= NewAnimSequence->CreateAnimation(Sequence);
-					break;
-				case EPoseSourceOption::CurrentAnimation_PreviewMesh:
-					bResult &= RecordMeshToAnimation(MeshComponent, NewAnimSequence);
-					break;
-				default: 
-					ensure(false);
-					break;
-				}
-			}
-		}
-
-		if (bResult)
-		{
-			HandleAssetCreated(NewAssets);
-
-			// if it created based on current mesh component, 
-			if (Option == EPoseSourceOption::CurrentPose)
-			{
-				UDebugSkelMeshComponent* PreviewMeshComponent = PersonaToolkit->GetPreviewMeshComponent();
-				if (PreviewMeshComponent && PreviewMeshComponent->PreviewInstance)
-				{
-					PreviewMeshComponent->PreviewInstance->ResetModifiedBone();
-				}
-			}
-		}
-	}
-	return true;
-}
-
-bool FAnimationEditor::CreatePoseAsset(const TArray<UObject*> NewAssets, const EPoseSourceOption Option)
-{
-	bool bResult = false;
-	if (NewAssets.Num() > 0)
-	{
-		UDebugSkelMeshComponent* PreviewComponent = PersonaToolkit->GetPreviewMeshComponent();
-		UAnimSequence* Sequence = Cast<UAnimSequence>(AnimationAsset);
-
-		for (auto NewAsset : NewAssets)
-		{
-			UPoseAsset* NewPoseAsset = Cast<UPoseAsset>(NewAsset);
-			if (NewPoseAsset)
-			{
-				switch (Option)
-				{
-				case EPoseSourceOption::CurrentPose:
-					NewPoseAsset->AddOrUpdatePoseWithUniqueName(PreviewComponent);
-					bResult = true;
-					break;
-				case EPoseSourceOption::CurrentAnimation_AnimData:
-					NewPoseAsset->CreatePoseFromAnimation(Sequence);
-					bResult = true;
-					break;
-				default:
-					ensure(false);
-					bResult = false; 
-					break;
-				}
-
-			}
-		}
-
-		// if it contains error, warn them
-		if (bResult)
-		{
-			HandleAssetCreated(NewAssets);
-
-			// if it created based on current mesh component, 
-			if (Option == EPoseSourceOption::CurrentPose)
-			{
-				PreviewComponent->PreviewInstance->ResetModifiedBone();
-			}
-		}
-		else
-		{
-			FMessageDialog::Open(EAppMsgType::Ok, LOCTEXT("FailedToCreateAsset", "Failed to create asset"));
-		}
-	}
-	return true;
-}
-
-bool FAnimationEditor::HandleAssetCreated(const TArray<UObject*> NewAssets)
-{
-	if (NewAssets.Num() > 0)
-	{
-		FAssetRegistryModule::AssetCreated(NewAssets[0]);
-
-		if (UAnimationAsset* NewAnimAsset = Cast<UAnimationAsset>(NewAssets[0]))
-		{
-			OpenNewAnimationDocumentTab(NewAnimAsset);
-		}
-		else
-		{
-			// if not, we forward to asset manager to open the asset for us
-			// this is path for animation blueprint
-			FAssetToolsModule& AssetToolsModule = FModuleManager::LoadModuleChecked<FAssetToolsModule>(TEXT("AssetTools"));
-			TWeakPtr<IAssetTypeActions> AssetTypeActions = AssetToolsModule.Get().GetAssetTypeActionsForClass(NewAssets[0]->GetClass());
-			if (AssetTypeActions.IsValid())
-			{
-				AssetTypeActions.Pin()->OpenAssetEditor(NewAssets);
-			}
-		}
-	}
-	return true;
-}
-
 void FAnimationEditor::ConditionalRefreshEditor(UObject* InObject)
 {
 	bool bInterestingAsset = true;
diff --git a/Engine/Source/Editor/AnimationEditor/Private/AnimationEditor.h b/Engine/Source/Editor/AnimationEditor/Private/AnimationEditor.h
--- a/Engine/Source/Editor/AnimationEditor/Private/AnimationEditor.h
+++ b/Engine/Source/Editor/AnimationEditor/Private/AnimationEditor.h
@@ -105,10 +105,9 @@ public:
 
 private:
 
-	enum class EPoseSourceOption : uint8
+	/** Options for asset export */
+	enum class EExportSourceOption : uint8
 	{
-		ReferencePose,
-		CurrentPose,
 		CurrentAnimation_AnimData,
 		CurrentAnimation_PreviewMesh,
 		Max
@@ -130,38 +129,21 @@ private:
 
 	void OnApplyCompression();
 
-	void OnExportToFBX(const EPoseSourceOption Option);
+	void OnExportToFBX(const EExportSourceOption Option);
 	//Return true mean the asset was exported, false it was cancel or it fail
 	bool ExportToFBX(const TArray<UObject*> NewAssets, bool bRecordAnimation);
 
 	void OnAddLoopingInterpolation();
 	void OnRemoveBoneTrack();
 
-	TSharedRef< SWidget > GenerateCreateAssetMenu() const;
 	TSharedRef< SWidget > GenerateExportAssetMenu() const;
 
-	void FillCreateAnimationMenu(FMenuBuilder& MenuBuilder) const;
-
-	void FillCreateAnimationFromCurrentAnimationMenu(FMenuBuilder& MenuBuilder) const;
-
-	void FillCreatePoseAssetMenu(FMenuBuilder& MenuBuilder) const;
-
-	void FillInsertPoseMenu(FMenuBuilder& MenuBuilder) const;
-
-	void InsertCurrentPoseToAsset(const FAssetData& NewPoseAssetData);
-
 	void FillCopyToSoundWaveMenu(FMenuBuilder& MenuBuilder) const;
 
 	void FillExportAssetMenu(FMenuBuilder& MenuBuilder) const;
 
 	void CopyCurveToSoundWave(const FAssetData& SoundWaveAssetData) const;
 
-	bool CreateAnimation(const TArray<UObject*> NewAssets, const EPoseSourceOption Option);
-
-	bool CreatePoseAsset(const TArray<UObject*> NewAssets, const EPoseSourceOption Option);
-
-	bool HandleAssetCreated(const TArray<UObject*> NewAssets);
-
 	void ConditionalRefreshEditor(UObject* InObject);
 
 	void HandlePostReimport(UObject* InObject, bool bSuccess);
diff --git a/Engine/Source/Editor/Persona/Persona.Build.cs b/Engine/Source/Editor/Persona/Persona.Build.cs
--- a/Engine/Source/Editor/Persona/Persona.Build.cs
+++ b/Engine/Source/Editor/Persona/Persona.Build.cs
@@ -61,6 +61,7 @@ public class Persona : ModuleRules
                 "UnrealEd",
                 "CommonMenuExtensions",
                 "PinnedCommandList",
+				"SequenceRecorder",
             }
         );
 
diff --git a/Engine/Source/Editor/Persona/Private/PersonaModule.cpp b/Engine/Source/Editor/Persona/Private/PersonaModule.cpp
--- a/Engine/Source/Editor/Persona/Private/PersonaModule.cpp
+++ b/Engine/Source/Editor/Persona/Private/PersonaModule.cpp
@@ -77,6 +77,12 @@
 #include "PersonaPreviewSceneRefPoseController.h"
 #include "AssetViewerSettings.h"
 #include "Customization/SkeletalMeshRegionCustomization.h"
+#include "AnimationEditorUtils.h"
+#include "Factories/AnimSequenceFactory.h"
+#include "Factories/PoseAssetFactory.h"
+#include "ScopedTransaction.h"
+#include "AnimPreviewInstance.h"
+#include "SequenceRecorderUtils.h"
 
 IMPLEMENT_MODULE( FPersonaModule, Persona );
 
@@ -953,6 +959,401 @@ void FPersonaModule::AddCommonToolbarExtensions(FToolBarBuilder& InToolbarBuilde
 			FSlateIcon("EditorStyle", "Persona.ToggleReferencePose", "Persona.ToggleReferencePose.Small")
 			);
 	}
+
+	if(InArgs.bCreateAsset)
+	{
+		InToolbarBuilder.AddComboButton(
+			FUIAction(),
+			FOnGetContent::CreateRaw(this, &FPersonaModule::GenerateCreateAssetMenu, WeakPersonaToolkit),
+			LOCTEXT("CreateAsset_Label", "Create Asset"),
+			LOCTEXT("CreateAsset_ToolTip", "Create Assets for this skeleton."),
+			FSlateIcon(FEditorStyle::GetStyleSetName(), "Persona.CreateAsset")
+		);
+	}
+}
+
+TSharedRef< SWidget > FPersonaModule::GenerateCreateAssetMenu(TWeakPtr<IPersonaToolkit> InWeakPersonaToolkit) const
+{
+	const bool bShouldCloseWindowAfterMenuSelection = true;
+	FMenuBuilder MenuBuilder(bShouldCloseWindowAfterMenuSelection, NULL);
+
+	// Create Animation menu
+	MenuBuilder.BeginSection("CreateAnimation", LOCTEXT("CreateAnimationMenuHeading", "Animation"));
+	{
+		// create menu
+		MenuBuilder.AddSubMenu(
+			LOCTEXT("CreateAnimationSubmenu", "Create Animation"),
+			LOCTEXT("CreateAnimationSubmenu_ToolTip", "Create Animation for this skeleton"),
+			FNewMenuDelegate::CreateRaw(this, &FPersonaModule::FillCreateAnimationMenu, InWeakPersonaToolkit),
+			false,
+			FSlateIcon(FEditorStyle::GetStyleSetName(), "Persona.AssetActions.CreateAnimAsset")
+			);
+
+		MenuBuilder.AddSubMenu(
+			LOCTEXT("CreatePoseAssetSubmenu", "Create PoseAsset"),
+			LOCTEXT("CreatePoseAsssetSubmenu_ToolTip", "Create PoseAsset for this skeleton"),
+			FNewMenuDelegate::CreateRaw(this, &FPersonaModule::FillCreatePoseAssetMenu, InWeakPersonaToolkit),
+			false,
+			FSlateIcon(FEditorStyle::GetStyleSetName(), "ClassIcon.PoseAsset")
+		);
+	}
+	MenuBuilder.EndSection();
+
+	TSharedRef<IPersonaToolkit> PersonaToolkit = InWeakPersonaToolkit.Pin().ToSharedRef();
+	TArray<TWeakObjectPtr<UObject>> Objects;
+	if (PersonaToolkit->GetPreviewMesh())
+	{
+		Objects.Add(PersonaToolkit->GetPreviewMesh());
+	}
+	else
+	{
+		Objects.Add(PersonaToolkit->GetSkeleton());
+	}
+
+	AnimationEditorUtils::FillCreateAssetMenu(MenuBuilder, Objects, FAnimAssetCreated::CreateRaw(this, &FPersonaModule::HandleAssetCreated), false);
+
+	return MenuBuilder.MakeWidget();
+}
+
+void FPersonaModule::FillCreateAnimationMenu(FMenuBuilder& MenuBuilder, TWeakPtr<IPersonaToolkit> InWeakPersonaToolkit) const
+{
+	TSharedRef<IPersonaToolkit> PersonaToolkit = InWeakPersonaToolkit.Pin().ToSharedRef();
+	TArray<TWeakObjectPtr<UObject>> Objects;
+	if (PersonaToolkit->GetPreviewMesh())
+	{
+		Objects.Add(PersonaToolkit->GetPreviewMesh());
+	}
+	else
+	{
+		Objects.Add(PersonaToolkit->GetSkeleton());
+	}
+
+	// create rig
+	MenuBuilder.BeginSection("CreateAnimationSubMenu", LOCTEXT("CreateAnimationSubMenuHeading", "Create Animation"));
+	{
+		MenuBuilder.AddMenuEntry(
+			LOCTEXT("CreateAnimation_RefPose", "Reference Pose"),
+			LOCTEXT("CreateAnimation_RefPose_Tooltip", "Create Animation from reference pose."),
+			FSlateIcon(),
+			FUIAction(
+				FExecuteAction::CreateStatic(&AnimationEditorUtils::ExecuteNewAnimAsset<UAnimSequenceFactory, UAnimSequence>, Objects, FString("_Sequence"), FAnimAssetCreated::CreateRaw(this, &FPersonaModule::CreateAnimation, EPoseSourceOption::ReferencePose, InWeakPersonaToolkit), false),
+				FCanExecuteAction()
+				)
+			);
+
+		MenuBuilder.AddMenuEntry(
+			LOCTEXT("CreateAnimation_CurrentPose", "Current Pose"),
+			LOCTEXT("CreateAnimation_CurrentPose_Tooltip", "Create Animation from current pose."),
+			FSlateIcon(),
+			FUIAction(
+				FExecuteAction::CreateStatic(&AnimationEditorUtils::ExecuteNewAnimAsset<UAnimSequenceFactory, UAnimSequence>, Objects, FString("_Sequence"), FAnimAssetCreated::CreateRaw(this, &FPersonaModule::CreateAnimation, EPoseSourceOption::CurrentPose, InWeakPersonaToolkit), false),
+				FCanExecuteAction()
+				)
+			);
+
+		if(UAnimSequence* AnimSequence = Cast<UAnimSequence>(PersonaToolkit->GetAnimationAsset()))
+		{
+			MenuBuilder.AddSubMenu(
+				LOCTEXT("CreateAnimation_CurrenAnimationSubMenu", "Current Animation"),
+				LOCTEXT("CreateAnimation_CurrenAnimationSubMenu_ToolTip", "Create Animation from current animation"),
+				FNewMenuDelegate::CreateRaw(this, &FPersonaModule::FillCreateAnimationFromCurrentAnimationMenu, InWeakPersonaToolkit),
+				false,
+				FSlateIcon(FEditorStyle::GetStyleSetName(), "Persona.AssetActions.CreateAnimAsset")
+			);
+		}
+	}
+	MenuBuilder.EndSection();
+}
+
+void FPersonaModule::FillCreateAnimationFromCurrentAnimationMenu(FMenuBuilder& MenuBuilder, TWeakPtr<IPersonaToolkit> InWeakPersonaToolkit) const
+{
+	TSharedRef<IPersonaToolkit> PersonaToolkit = InWeakPersonaToolkit.Pin().ToSharedRef();
+	TArray<TWeakObjectPtr<UObject>> Objects;
+
+	if (PersonaToolkit->GetPreviewMesh())
+	{
+		Objects.Add(PersonaToolkit->GetPreviewMesh());
+	}
+	else
+	{
+		Objects.Add(PersonaToolkit->GetSkeleton());
+	}
+
+	// create rig
+	MenuBuilder.BeginSection("CreateAnimationSubMenu", LOCTEXT("CreateAnimationFromCurrentAnimationSubmenuHeading", "Create Animation"));
+	{
+		MenuBuilder.AddMenuEntry(
+			LOCTEXT("CreateAnimation_CurrentAnimation_AnimData", "Animation Data"),
+			LOCTEXT("CreateAnimation_CurrentAnimation_AnimData_Tooltip", "Create Animation from Animation Source Data."),
+			FSlateIcon(),
+			FUIAction(
+				FExecuteAction::CreateStatic(&AnimationEditorUtils::ExecuteNewAnimAsset<UAnimSequenceFactory, UAnimSequence>, Objects, FString("_Sequence"), FAnimAssetCreated::CreateRaw(this, &FPersonaModule::CreateAnimation, EPoseSourceOption::CurrentAnimation_AnimData, InWeakPersonaToolkit), false)
+			)
+		);
+
+		MenuBuilder.AddMenuEntry(
+			LOCTEXT("CreateAnimation_CurrentAnimation_PreviewMesh", "Preview Mesh"),
+			LOCTEXT("CreateAnimation_CurrentAnimation_PreviewMesh_Tooltip", "Create Animation by playing on the Current Preview Mesh, including Retargeting, Post Process Graph, or anything you see on the preview mesh."),
+			FSlateIcon(),
+			FUIAction(
+				FExecuteAction::CreateStatic(&AnimationEditorUtils::ExecuteNewAnimAsset<UAnimSequenceFactory, UAnimSequence>, Objects, FString("_Sequence"), FAnimAssetCreated::CreateRaw(this, &FPersonaModule::CreateAnimation, EPoseSourceOption::CurrentAnimation_PreviewMesh, InWeakPersonaToolkit), false)
+			)
+		);
+	}
+	MenuBuilder.EndSection();
+}
+
+void FPersonaModule::FillCreatePoseAssetMenu(FMenuBuilder& MenuBuilder, TWeakPtr<IPersonaToolkit> InWeakPersonaToolkit) const
+{
+	TSharedRef<IPersonaToolkit> PersonaToolkit = InWeakPersonaToolkit.Pin().ToSharedRef();
+	TArray<TWeakObjectPtr<UObject>> Objects;
+
+	if (PersonaToolkit->GetPreviewMesh())
+	{
+		Objects.Add(PersonaToolkit->GetPreviewMesh());
+	}
+	else
+	{
+		Objects.Add(PersonaToolkit->GetSkeleton());
+	}
+
+	// create rig
+	MenuBuilder.BeginSection("CreatePoseAssetSubMenu", LOCTEXT("CreatePoseAssetSubMenuHeading", "Create PoseAsset"));
+	{
+		MenuBuilder.AddMenuEntry(
+			LOCTEXT("CreatePoseAsset_CurrentPose", "Current Pose"),
+			LOCTEXT("CreatePoseAsset_CurrentPose_Tooltip", "Create PoseAsset from current pose."),
+			FSlateIcon(),
+			FUIAction(
+				FExecuteAction::CreateStatic(&AnimationEditorUtils::ExecuteNewAnimAsset<UPoseAssetFactory, UPoseAsset>, Objects, FString("_PoseAsset"), FAnimAssetCreated::CreateRaw(this, &FPersonaModule::CreatePoseAsset, EPoseSourceOption::CurrentPose, InWeakPersonaToolkit), false),
+				FCanExecuteAction()
+			)
+		);
+
+		MenuBuilder.AddMenuEntry(
+			LOCTEXT("CreatePoseAsset_CurrentAnimation", "Current Animation"),
+			LOCTEXT("CreatePoseAsset_CurrentAnimation_Tooltip", "Create Animation from current animation."),
+			FSlateIcon(),
+			FUIAction(
+				FExecuteAction::CreateStatic(&AnimationEditorUtils::ExecuteNewAnimAsset<UPoseAssetFactory, UPoseAsset>, Objects, FString("_PoseAsset"), FAnimAssetCreated::CreateRaw(this, &FPersonaModule::CreatePoseAsset, EPoseSourceOption::CurrentAnimation_AnimData, InWeakPersonaToolkit), false),
+				FCanExecuteAction()
+			)
+		);
+	}
+	MenuBuilder.EndSection();
+
+	// create pose asset
+	MenuBuilder.BeginSection("InsertPoseSubMenuSection", LOCTEXT("InsertPoseSubMenuSubMenuHeading", "Insert Pose"));
+	{
+		MenuBuilder.AddSubMenu(
+			LOCTEXT("InsertPoseSubmenu", "Insert Pose"),
+			LOCTEXT("InsertPoseSubmenu_ToolTip", "Insert current pose to selected PoseAsset"),
+			FNewMenuDelegate::CreateRaw(this, &FPersonaModule::FillInsertPoseMenu, InWeakPersonaToolkit),
+			false,
+			FSlateIcon(FEditorStyle::GetStyleSetName(), "ClassIcon.PoseAsset")
+		);
+	}
+	MenuBuilder.EndSection();
+}
+
+void FPersonaModule::FillInsertPoseMenu(FMenuBuilder& MenuBuilder, TWeakPtr<IPersonaToolkit> InWeakPersonaToolkit) const
+{
+	FAssetPickerConfig AssetPickerConfig;
+
+	TSharedRef<IPersonaToolkit> PersonaToolkit = InWeakPersonaToolkit.Pin().ToSharedRef();
+	USkeleton* Skeleton = PersonaToolkit->GetSkeleton();
+
+	/** The asset picker will only show skeletons */
+	AssetPickerConfig.Filter.ClassNames.Add(*UPoseAsset::StaticClass()->GetName());
+	AssetPickerConfig.Filter.bRecursiveClasses = false;
+	AssetPickerConfig.bAllowNullSelection = false;
+	AssetPickerConfig.Filter.TagsAndValues.Add(TEXT("Skeleton"), FAssetData(Skeleton).GetExportTextName());
+
+	/** The delegate that fires when an asset was selected */
+	AssetPickerConfig.OnAssetSelected = FOnAssetSelected::CreateRaw(this, &FPersonaModule::InsertCurrentPoseToAsset, InWeakPersonaToolkit);
+
+	/** The default view mode should be a list view */
+	AssetPickerConfig.InitialAssetViewType = EAssetViewType::List;
+
+	FContentBrowserModule& ContentBrowserModule = FModuleManager::Get().LoadModuleChecked<FContentBrowserModule>(TEXT("ContentBrowser"));
+
+	MenuBuilder.AddWidget(
+		ContentBrowserModule.Get().CreateAssetPicker(AssetPickerConfig),
+		FText::GetEmpty()
+	);
+}
+
+void FPersonaModule::InsertCurrentPoseToAsset(const FAssetData& NewPoseAssetData, TWeakPtr<IPersonaToolkit> InWeakPersonaToolkit)
+{
+	TSharedRef<IPersonaToolkit> PersonaToolkit = InWeakPersonaToolkit.Pin().ToSharedRef();
+	UPoseAsset* PoseAsset = Cast<UPoseAsset>(NewPoseAssetData.GetAsset());
+	FScopedTransaction ScopedTransaction(LOCTEXT("InsertPose", "Insert Pose"));
+
+	if (PoseAsset)
+	{
+		PoseAsset->Modify();
+
+		UDebugSkelMeshComponent* PreviewMeshComponent = PersonaToolkit->GetPreviewMeshComponent();
+		if (PreviewMeshComponent)
+		{
+			FSmartName NewPoseName;
+
+			bool bSuccess = PoseAsset->AddOrUpdatePoseWithUniqueName(PreviewMeshComponent, &NewPoseName);
+
+			if (bSuccess)
+			{
+				FFormatNamedArguments Args;
+				Args.Add(TEXT("PoseAsset"), FText::FromString(PoseAsset->GetName()));
+				Args.Add(TEXT("PoseName"), FText::FromName(NewPoseName.DisplayName));
+				FNotificationInfo Info(FText::Format(LOCTEXT("InsertPoseSucceeded", "The current pose has inserted to {PoseAsset} with {PoseName}"), Args));
+				Info.ExpireDuration = 7.0f;
+				Info.bUseLargeFont = false;
+				TSharedPtr<SNotificationItem> Notification = FSlateNotificationManager::Get().AddNotification(Info);
+				if (Notification.IsValid())
+				{
+					Notification->SetCompletionState(SNotificationItem::CS_Success);
+				}
+			}
+			else
+			{
+				FFormatNamedArguments Args;
+				Args.Add(TEXT("PoseAsset"), FText::FromString(PoseAsset->GetName()));
+				FNotificationInfo Info(FText::Format(LOCTEXT("InsertPoseFailed", "Inserting pose to asset {PoseAsset} has failed"), Args));
+				Info.ExpireDuration = 7.0f;
+				Info.bUseLargeFont = false;
+				TSharedPtr<SNotificationItem> Notification = FSlateNotificationManager::Get().AddNotification(Info);
+				if (Notification.IsValid())
+				{
+					Notification->SetCompletionState(SNotificationItem::CS_Fail);
+				}
+			}
+		}
+	}
+
+	// it doesn't work well if I leave the window open. The delegate goes weired or it stop showing the popups. 
+	FSlateApplication::Get().DismissAllMenus();
+}
+
+bool FPersonaModule::CreateAnimation(const TArray<UObject*> NewAssets, const EPoseSourceOption Option, TWeakPtr<IPersonaToolkit> InWeakPersonaToolkit)
+{
+	bool bResult = true;
+	if (NewAssets.Num() > 0)
+	{
+		TSharedRef<IPersonaToolkit> PersonaToolkit = InWeakPersonaToolkit.Pin().ToSharedRef();
+		USkeletalMeshComponent* MeshComponent = PersonaToolkit->GetPreviewMeshComponent();
+		UAnimSequence* Sequence = Cast<UAnimSequence>(PersonaToolkit->GetAnimationAsset());
+
+		for (auto NewAsset : NewAssets)
+		{
+			UAnimSequence* NewAnimSequence = Cast<UAnimSequence>(NewAsset);
+			if (NewAnimSequence)
+			{
+				switch (Option)
+				{
+				case EPoseSourceOption::ReferencePose:
+					bResult &= NewAnimSequence->CreateAnimation(MeshComponent->SkeletalMesh);
+					break;
+				case EPoseSourceOption::CurrentPose:
+					bResult &= NewAnimSequence->CreateAnimation(MeshComponent);
+					break;
+				case EPoseSourceOption::CurrentAnimation_AnimData:
+					bResult &= NewAnimSequence->CreateAnimation(Sequence);
+					break;
+				case EPoseSourceOption::CurrentAnimation_PreviewMesh:
+					bResult &= SequenceRecorderUtils::RecordSingleNodeInstanceToAnimation(MeshComponent, NewAnimSequence);
+					break;
+				default: 
+					ensure(false);
+					break;
+				}
+			}
+		}
+
+		if (bResult)
+		{
+			HandleAssetCreated(NewAssets);
+
+			// if it created based on current mesh component, 
+			if (Option == EPoseSourceOption::CurrentPose)
+			{
+				UDebugSkelMeshComponent* PreviewMeshComponent = PersonaToolkit->GetPreviewMeshComponent();
+				if (PreviewMeshComponent && PreviewMeshComponent->PreviewInstance)
+				{
+					PreviewMeshComponent->PreviewInstance->ResetModifiedBone();
+				}
+			}
+		}
+	}
+	return true;
+}
+
+bool FPersonaModule::CreatePoseAsset(const TArray<UObject*> NewAssets, const EPoseSourceOption Option, TWeakPtr<IPersonaToolkit> InWeakPersonaToolkit)
+{
+	bool bResult = false;
+	if (NewAssets.Num() > 0)
+	{
+		TSharedRef<IPersonaToolkit> PersonaToolkit = InWeakPersonaToolkit.Pin().ToSharedRef();
+		UDebugSkelMeshComponent* PreviewComponent = PersonaToolkit->GetPreviewMeshComponent();
+		UAnimSequence* Sequence = Cast<UAnimSequence>(PersonaToolkit->GetAnimationAsset());
+
+		for (auto NewAsset : NewAssets)
+		{
+			UPoseAsset* NewPoseAsset = Cast<UPoseAsset>(NewAsset);
+			if (NewPoseAsset)
+			{
+				switch (Option)
+				{
+				case EPoseSourceOption::CurrentPose:
+					NewPoseAsset->AddOrUpdatePoseWithUniqueName(PreviewComponent);
+					bResult = true;
+					break;
+				case EPoseSourceOption::CurrentAnimation_AnimData:
+					NewPoseAsset->CreatePoseFromAnimation(Sequence);
+					bResult = true;
+					break;
+				default:
+					ensure(false);
+					bResult = false; 
+					break;
+				}
+
+			}
+		}
+
+		// if it contains error, warn them
+		if (bResult)
+		{
+			HandleAssetCreated(NewAssets);
+
+			// if it created based on current mesh component, 
+			if (Option == EPoseSourceOption::CurrentPose)
+			{
+				PreviewComponent->PreviewInstance->ResetModifiedBone();
+			}
+		}
+		else
+		{
+			FMessageDialog::Open(EAppMsgType::Ok, LOCTEXT("FailedToCreateAsset", "Failed to create asset"));
+		}
+	}
+	return true;
+}
+
+bool FPersonaModule::HandleAssetCreated(const TArray<UObject*> NewAssets)
+{
+	if (NewAssets.Num() > 0)
+	{
+		FAssetRegistryModule::AssetCreated(NewAssets[0]);
+
+		// forward to asset manager to open the asset for us
+		FAssetToolsModule& AssetToolsModule = FModuleManager::LoadModuleChecked<FAssetToolsModule>(TEXT("AssetTools"));
+		TWeakPtr<IAssetTypeActions> AssetTypeActions = AssetToolsModule.Get().GetAssetTypeActionsForClass(NewAssets[0]->GetClass());
+		if (AssetTypeActions.IsValid())
+		{
+			AssetTypeActions.Pin()->OpenAssetEditor(NewAssets);
+		}
+	}
+	return true;
 }
 
 void FPersonaModule::HandleNewAnimNotifyBlueprintCreated(UBlueprint* InBlueprint)
diff --git a/Engine/Source/Editor/Persona/Public/PersonaModule.h b/Engine/Source/Editor/Persona/Public/PersonaModule.h
--- a/Engine/Source/Editor/Persona/Public/PersonaModule.h
+++ b/Engine/Source/Editor/Persona/Public/PersonaModule.h
@@ -300,6 +300,7 @@ public:
 			: bPreviewMesh(true)
 			, bPreviewAnimation(true)
 			, bReferencePose(false)
+			, bCreateAsset(true)
 		{}
 
 		/** Adds a shortcut to setup a preview mesh to override the current display */
@@ -310,6 +311,9 @@ public:
 
 		/** Adds a shortcut to set the character back to reference pose (also clears all bone modifications) */
 		bool bReferencePose;
+
+		/** Adds a combo menu to allow other anim assets to be created */
+		bool bCreateAsset;
 	};
 
 	/** Add common toobar extensions */
@@ -328,6 +332,34 @@ private:
 	/** When a new anim notify state blueprint is created, this will handle post creation work such as adding non-event default nodes */
 	void HandleNewAnimNotifyStateBlueprintCreated(UBlueprint* InBlueprint);
 
+	/** Options for asset creation */
+	enum class EPoseSourceOption : uint8
+	{
+		ReferencePose,
+		CurrentPose,
+		CurrentAnimation_AnimData,
+		CurrentAnimation_PreviewMesh,
+		Max
+	};
+
+	TSharedRef< SWidget > GenerateCreateAssetMenu(TWeakPtr<IPersonaToolkit> InWeakPersonaToolkit) const;
+
+	void FillCreateAnimationMenu(FMenuBuilder& MenuBuilder, TWeakPtr<IPersonaToolkit> InWeakPersonaToolkit) const;
+
+	void FillCreateAnimationFromCurrentAnimationMenu(FMenuBuilder& MenuBuilder, TWeakPtr<IPersonaToolkit> InWeakPersonaToolkit) const;
+
+	void FillCreatePoseAssetMenu(FMenuBuilder& MenuBuilder, TWeakPtr<IPersonaToolkit> InWeakPersonaToolkit) const;
+
+	void FillInsertPoseMenu(FMenuBuilder& MenuBuilder, TWeakPtr<IPersonaToolkit> InWeakPersonaToolkit) const;
+
+	void InsertCurrentPoseToAsset(const FAssetData& NewPoseAssetData, TWeakPtr<IPersonaToolkit> InWeakPersonaToolkit);
+
+	bool CreateAnimation(const TArray<UObject*> NewAssets, const EPoseSourceOption Option, TWeakPtr<IPersonaToolkit> InWeakPersonaToolkit);
+
+	bool CreatePoseAsset(const TArray<UObject*> NewAssets, const EPoseSourceOption Option, TWeakPtr<IPersonaToolkit> InWeakPersonaToolkit);
+	
+	bool HandleAssetCreated(const TArray<UObject*> NewAssets);
+
 private:
 	TSharedPtr<FExtensibilityManager> MenuExtensibilityManager;
 	TSharedPtr<FExtensibilityManager> ToolBarExtensibilityManager;
