//UE4/Dev-Anim - Incremental EngineTestEditor Mac - Action graph contains cycle!

diff --git a/Engine/Source/Editor/AnimationEditor/AnimationEditor.Build.cs b/Engine/Source/Editor/AnimationEditor/AnimationEditor.Build.cs
--- a/Engine/Source/Editor/AnimationEditor/AnimationEditor.Build.cs
+++ b/Engine/Source/Editor/AnimationEditor/AnimationEditor.Build.cs
@@ -20,19 +20,20 @@ public class AnimationEditor : ModuleRules
                 "SkeletonEditor",
                 "Kismet",
                 "AnimGraph",
-                "SequenceRecorder",
             }
 		);
 
         PrivateIncludePathModuleNames.AddRange(
             new string[] {
                 "PropertyEditor",
+                "SequenceRecorder",
             }
         );
 
         DynamicallyLoadedModuleNames.AddRange(
             new string[] {
                 "PropertyEditor",
+                "SequenceRecorder",
             }
         );
     }
diff --git a/Engine/Source/Editor/AnimationEditor/Private/AnimationEditor.cpp b/Engine/Source/Editor/AnimationEditor/Private/AnimationEditor.cpp
--- a/Engine/Source/Editor/AnimationEditor/Private/AnimationEditor.cpp
+++ b/Engine/Source/Editor/AnimationEditor/Private/AnimationEditor.cpp
@@ -46,9 +46,9 @@
 #include "Engine/CurveTable.h"
 #include "Developer/AssetTools/Public/IAssetTools.h"
 #include "Developer/AssetTools/Public/AssetToolsModule.h"
-#include "SequenceRecorderUtils.h"
 #include "ISkeletonTreeItem.h"
 #include "Algo/Transform.h"
+#include "ISequenceRecorder.h"
 
 const FName AnimationEditorAppIdentifier = FName(TEXT("AnimationEditorApp"));
 
@@ -783,7 +783,8 @@ void FAnimationEditor::HandleAnimationSequenceBrowserCreated(const TSharedRef<IA
 
 bool FAnimationEditor::RecordMeshToAnimation(USkeletalMeshComponent* PreviewComponent, UAnimSequence* NewAsset) const
 {
-	return SequenceRecorderUtils::RecordSingleNodeInstanceToAnimation(PreviewComponent, NewAsset);
+	ISequenceRecorder& RecorderModule = FModuleManager::Get().LoadModuleChecked<ISequenceRecorder>("SequenceRecorder");
+	return RecorderModule.RecordSingleNodeInstanceToAnimation(PreviewComponent, NewAsset);
 }
 
 #undef LOCTEXT_NAMESPACE
diff --git a/Engine/Source/Editor/Persona/Persona.Build.cs b/Engine/Source/Editor/Persona/Persona.Build.cs
--- a/Engine/Source/Editor/Persona/Persona.Build.cs
+++ b/Engine/Source/Editor/Persona/Persona.Build.cs
@@ -28,7 +28,8 @@ public class Persona : ModuleRules
                 "ContentBrowser",
                 "AssetTools",
                 "AnimationEditor",
-                "MeshReductionInterface"
+                "MeshReductionInterface",
+                "SequenceRecorder",
             }
         );
 
@@ -61,7 +62,6 @@ public class Persona : ModuleRules
                 "UnrealEd",
                 "CommonMenuExtensions",
                 "PinnedCommandList",
-				"SequenceRecorder",
             }
         );
 
@@ -74,7 +74,8 @@ public class Persona : ModuleRules
                 "SkeletonEditor",
                 "AssetTools",
                 "AnimationEditor",
-                "MeshReductionInterface"
+                "MeshReductionInterface",
+                "SequenceRecorder",
             }
         );
     }
diff --git a/Engine/Source/Editor/Persona/Private/PersonaModule.cpp b/Engine/Source/Editor/Persona/Private/PersonaModule.cpp
--- a/Engine/Source/Editor/Persona/Private/PersonaModule.cpp
+++ b/Engine/Source/Editor/Persona/Private/PersonaModule.cpp
@@ -82,7 +82,7 @@
 #include "Factories/PoseAssetFactory.h"
 #include "ScopedTransaction.h"
 #include "AnimPreviewInstance.h"
-#include "SequenceRecorderUtils.h"
+#include "ISequenceRecorder.h"
 
 IMPLEMENT_MODULE( FPersonaModule, Persona );
 
@@ -1260,7 +1260,10 @@ bool FPersonaModule::CreateAnimation(const TArray<UObject*> NewAssets, const EPo
 					bResult &= NewAnimSequence->CreateAnimation(Sequence);
 					break;
 				case EPoseSourceOption::CurrentAnimation_PreviewMesh:
-					bResult &= SequenceRecorderUtils::RecordSingleNodeInstanceToAnimation(MeshComponent, NewAnimSequence);
+				{
+					ISequenceRecorder & RecorderModule = FModuleManager::Get().LoadModuleChecked<ISequenceRecorder>("SequenceRecorder");
+					bResult &= RecorderModule.RecordSingleNodeInstanceToAnimation(MeshComponent, NewAnimSequence);
+				}
 					break;
 				default: 
 					ensure(false);
diff --git a/Engine/Source/Editor/SequenceRecorder/Private/SequenceRecorderModule.cpp b/Engine/Source/Editor/SequenceRecorder/Private/SequenceRecorderModule.cpp
--- a/Engine/Source/Editor/SequenceRecorder/Private/SequenceRecorderModule.cpp
+++ b/Engine/Source/Editor/SequenceRecorder/Private/SequenceRecorderModule.cpp
@@ -44,6 +44,7 @@
 #include "WorkflowOrientedApp/WorkflowTabManager.h"
 #include "Framework/Docking/LayoutExtender.h"
 #include "MovieSceneTimeHelpers.h"
+#include "SequenceRecorderUtils.h"
 
 #define LOCTEXT_NAMESPACE "SequenceRecorder"
 
@@ -773,6 +774,11 @@ class FSequenceRecorderModule : public ISequenceRecorder, private FSelfRegisteri
 		}
 	}
 
+	virtual bool RecordSingleNodeInstanceToAnimation(USkeletalMeshComponent* PreviewComponent, UAnimSequence* NewAsset) override
+	{
+		return SequenceRecorderUtils::RecordSingleNodeInstanceToAnimation(PreviewComponent, NewAsset);
+	}
+
 #if WITH_EDITOR
 	static UAnimSequence* HandleCaptureSingleFrameAnimSequence(USkeletalMeshComponent* Component)
 	{
diff --git a/Engine/Source/Editor/SequenceRecorder/Public/ISequenceRecorder.h b/Engine/Source/Editor/SequenceRecorder/Public/ISequenceRecorder.h
--- a/Engine/Source/Editor/SequenceRecorder/Public/ISequenceRecorder.h
+++ b/Engine/Source/Editor/SequenceRecorder/Public/ISequenceRecorder.h
@@ -165,4 +165,12 @@ public:
 
 	/** Remove an extension from the SequenceRecorder */
 	virtual void RemoveSequenceRecorderExtender(TSharedPtr<ISequenceRecorderExtender> SequenceRecorderExternder) = 0;
+
+	/**
+	* Play the current single node instance on the PreviewComponent from time [0, GetLength()], and record to NewAsset
+	*
+	* @param: PreviewComponent - this component should contains SingleNodeInstance with time-line based asset, currently support AnimSequence or AnimComposite
+	* @param: NewAsset - this is the asset that should be recorded. This will reset all animation data internally
+	*/
+	virtual bool RecordSingleNodeInstanceToAnimation(USkeletalMeshComponent* PreviewComponent, UAnimSequence* NewAsset) = 0;
 };
