Copying //UE4/Dev-Anim to Dev-Main (//UE4/Dev-Main) @ 4803512
#rnx

diff --git a/Engine/Plugins/Developer/AnimationSharing/Source/AnimationSharing/Private/AdditiveAnimationInstance.cpp b/Engine/Plugins/Developer/AnimationSharing/Source/AnimationSharing/Private/AdditiveAnimationInstance.cpp
--- a/Engine/Plugins/Developer/AnimationSharing/Source/AnimationSharing/Private/AdditiveAnimationInstance.cpp
+++ b/Engine/Plugins/Developer/AnimationSharing/Source/AnimationSharing/Private/AdditiveAnimationInstance.cpp
@@ -10,14 +10,18 @@ FAdditiveAnimationInstance::FAdditiveAnimationInstance() : SkeletalMeshComponent
 
 void FAdditiveAnimationInstance::Initialise(USkeletalMeshComponent* InSkeletalMeshComponent, UClass* InAnimationBPClass)
 {
-	if (InSkeletalMeshComponent && InAnimationBPClass)
+	if (InSkeletalMeshComponent)
 	{
 		SkeletalMeshComponent = InSkeletalMeshComponent;
-		SkeletalMeshComponent->SetAnimInstanceClass(InAnimationBPClass);
+		if (InAnimationBPClass)
+		{
+			SkeletalMeshComponent->SetAnimInstanceClass(InAnimationBPClass);
+			AdditiveInstance = Cast<UAnimSharingAdditiveInstance>(SkeletalMeshComponent->GetAnimInstance());
+		}
+		
 		SkeletalMeshComponent->SetComponentTickEnabled(false);	
 		SkeletalMeshComponent->SetForcedLOD(1);
 		SkeletalMeshComponent->VisibilityBasedAnimTickOption = EVisibilityBasedAnimTickOption::AlwaysTickPoseAndRefreshBones;
-		AdditiveInstance = Cast<UAnimSharingAdditiveInstance>(SkeletalMeshComponent->GetAnimInstance());
 	}
 }
 
diff --git a/Engine/Plugins/Developer/AnimationSharing/Source/AnimationSharing/Private/AnimationSharingManager.cpp b/Engine/Plugins/Developer/AnimationSharing/Source/AnimationSharing/Private/AnimationSharingManager.cpp
--- a/Engine/Plugins/Developer/AnimationSharing/Source/AnimationSharing/Private/AnimationSharingManager.cpp
+++ b/Engine/Plugins/Developer/AnimationSharing/Source/AnimationSharing/Private/AnimationSharingManager.cpp
@@ -221,14 +221,22 @@ const FAnimationSharingScalability& UAnimationSharingManager::GetScalabilitySett
 void UAnimationSharingManager::SetupPerSkeletonData(const FPerSkeletonAnimationSharingSetup& SkeletonSetup)
 {
 	const USkeleton* Skeleton = SkeletonSetup.Skeleton.LoadSynchronous();
-	UEnum* StateEnum = SkeletonSetup.StateProcessorClass->GetDefaultObject<UAnimationSharingStateProcessor>()->GetAnimationStateEnum();
-	if (Skeleton && StateEnum)
+	UAnimationSharingStateProcessor* Processor = SkeletonSetup.StateProcessorClass ?SkeletonSetup.StateProcessorClass->GetDefaultObject<UAnimationSharingStateProcessor>() : nullptr;
+	UEnum* StateEnum = Processor ? Processor->GetAnimationStateEnum() : nullptr;
+	if (Skeleton && StateEnum && Processor)
 	{
 		UAnimSharingInstance* Data = NewObject<UAnimSharingInstance>(this);
 		PerSkeletonData.Add(Data);
 		Skeletons.Add(Skeleton);
 		Data->Setup(this, SkeletonSetup, &ScalabilitySettings, Skeletons.Num() - 1);
 	}
+	else
+	{
+		UE_LOG(LogAnimationSharing, Error, TEXT("Invalid Skeleton (%s), State Enum (%s) or State Processor (%s)!"), 
+			Skeleton ? *Skeleton->GetName() : TEXT("None"),
+			StateEnum ? *StateEnum->GetName() : TEXT("None"),
+			Processor ? *Processor->GetName() : TEXT("None"));
+	}
 }
 
 uint32 UAnimationSharingManager::CreateActorHandle(uint8 SkeletonIndex, uint32 ActorIndex) const
@@ -345,7 +353,7 @@ void UAnimationSharingManager::RegisterActor(AActor* InActor, FUpdateActorHandle
 		InActor->GetComponents<USkeletalMeshComponent>(OwnedComponents);
 		checkf(OwnedComponents.Num(), TEXT("No SkeletalMeshComponents found in actor!"));
 
-		const USkeleton* UsedSkeleton = [&]()
+		const USkeleton* UsedSkeleton = [&OwnedComponents]()
 		{
 			USkeleton* CurrentSkeleton = nullptr;
 			for (USkeletalMeshComponent* SkeletalMeshComponent : OwnedComponents)
@@ -380,13 +388,12 @@ void UAnimationSharingManager::RegisterActorWithSkeleton(AActor* InActor, const
 		return;
 	}
 
-	const AnimationSharingDataHandle Handle = [&]() -> uint32
+	const AnimationSharingDataHandle Handle = [this, SharingSkeleton]() -> uint32
 	{
-		uint32 ArrayIndex = Skeletons.IndexOfByPredicate([&](const USkeleton* Skeleton)
+		uint32 ArrayIndex = Skeletons.IndexOfByPredicate([SharingSkeleton](const USkeleton* Skeleton)
 		{
 			return (Skeleton == SharingSkeleton) || (Skeleton->IsCompatible(SharingSkeleton));
 		});
-		ensureMsgf(ArrayIndex != INDEX_NONE, TEXT("Invalid skeleton for which there is no sharing setup available!"));
 		return ArrayIndex;
 	}();
 
@@ -397,43 +404,49 @@ void UAnimationSharingManager::RegisterActorWithSkeleton(AActor* InActor, const
 		checkf(OwnedComponents.Num(), TEXT("No SkeletalMeshComponents found in actor!"));
 
 		UAnimSharingInstance* Data = PerSkeletonData[Handle];
+		if (Data->AnimSharingManager != nullptr)
+		{
+			// Register the actor
+			const int32 ActorIndex = Data->RegisteredActors.Add(InActor);
 
-		// Register the actor
-		const int32 ActorIndex = Data->RegisteredActors.Add(InActor);
+			FPerActorData& ActorData = Data->PerActorData.AddZeroed_GetRef();
+			ActorData.BlendInstanceIndex = ActorData.OnDemandInstanceIndex = ActorData.AdditiveInstanceIndex = INDEX_NONE;
+			ActorData.SignificanceValue = Data->SignificanceManager->GetSignificance(InActor);
+			ActorData.UpdateActorHandleDelegate = CallbackDelegate;
 
-		FPerActorData& ActorData = Data->PerActorData.AddZeroed_GetRef();
-		ActorData.BlendInstanceIndex = ActorData.OnDemandInstanceIndex = ActorData.AdditiveInstanceIndex = INDEX_NONE;
-		ActorData.SignificanceValue = Data->SignificanceManager->GetSignificance(InActor);
-		ActorData.UpdateActorHandleDelegate = CallbackDelegate;
+			bool bShouldProcess = true;
+			ActorData.CurrentState = ActorData.PreviousState = Data->DetermineStateForActor(ActorIndex, bShouldProcess);
 
-		bool bShouldProcess = true;
-		ActorData.CurrentState = ActorData.PreviousState = Data->DetermineStateForActor(ActorIndex, bShouldProcess);
+			for (USkeletalMeshComponent* Component : OwnedComponents)
+			{
+				FPerComponentData& ComponentData = Data->PerComponentData.AddZeroed_GetRef();
+				ComponentData.ActorIndex = ActorIndex;
+				ComponentData.Component = Component;
 
-		for (USkeletalMeshComponent* Component : OwnedComponents)
-		{
-			FPerComponentData& ComponentData = Data->PerComponentData.AddZeroed_GetRef();
-			ComponentData.ActorIndex = ActorIndex;
-			ComponentData.Component = Component;
+				Component->PrimaryComponentTick.bCanEverTick = false;
+				Component->SetComponentTickEnabled(false);
+				Component->bIgnoreMasterPoseComponentLOD = true;
 
-			Component->PrimaryComponentTick.bCanEverTick = false;
-			Component->SetComponentTickEnabled(false);
-			Component->bIgnoreMasterPoseComponentLOD = true;
+				ActorData.ComponentIndices.Add(Data->PerComponentData.Num() - 1);
 
-			ActorData.ComponentIndices.Add(Data->PerComponentData.Num() - 1);
+				const int32 ComponentIndex = Data->PerComponentData.Num() - 1;
+				Data->SetupSlaveComponent(ActorData.CurrentState, ActorIndex);
+			}
 
-			const int32 ComponentIndex = Data->PerComponentData.Num() - 1;
-			Data->SetupSlaveComponent(ActorData.CurrentState, ActorIndex);
-		}
+			if (Data->PerStateData[ActorData.CurrentState].bIsOnDemand && ActorData.OnDemandInstanceIndex != INDEX_NONE)
+			{
+				// We will have setup an on-demand instance so we need to kick it off here before we next tick
+				Data->OnDemandInstances[ActorData.OnDemandInstanceIndex].bActive = true;
+				Data->OnDemandInstances[ActorData.OnDemandInstanceIndex].StartTime = Data->WorldTime;
+			}
 
-		if (Data->PerStateData[ActorData.CurrentState].bIsOnDemand && ActorData.OnDemandInstanceIndex != INDEX_NONE)
-		{
-			// We will have setup an on-demand instance so we need to kick it off here before we next tick
-			Data->OnDemandInstances[ActorData.OnDemandInstanceIndex].bActive = true;
-			Data->OnDemandInstances[ActorData.OnDemandInstanceIndex].StartTime = Data->WorldTime;
+			const int32 ActorHandle = CreateActorHandle(Handle, ActorIndex);
+			ActorData.UpdateActorHandleDelegate.ExecuteIfBound(ActorHandle);
 		}
-
-		const int32 ActorHandle = CreateActorHandle(Handle, ActorIndex);
-		ActorData.UpdateActorHandleDelegate.ExecuteIfBound(ActorHandle);
+	}
+	else
+	{
+		UE_LOG(LogAnimationSharing, Error, TEXT("Invalid skeleton (%s) for which there is no sharing setup available!"), SharingSkeleton ? *SharingSkeleton->GetName() : TEXT("None"));
 	}
 }
 
@@ -603,22 +616,34 @@ void UAnimationSharingManager::SetMasterComponentsVisibility(bool bVisible)
 
 		for (FTransitionBlendInstance* Instance : Data->BlendInstanceStack.AvailableInstances)
 		{
-			Instance->GetComponent()->SetVisibility(bVisible);
+			if (USceneComponent* Component = Instance->GetComponent())
+			{
+				Component->SetVisibility(bVisible);
+			}
 		}
 
 		for (FTransitionBlendInstance* Instance : Data->BlendInstanceStack.InUseInstances)
 		{
-			Instance->GetComponent()->SetVisibility(bVisible);
+			if (USceneComponent* Component = Instance->GetComponent())
+			{
+				Component->SetVisibility(bVisible);
+			}
 		}
 
 		for (FAdditiveAnimationInstance* Instance : Data->AdditiveInstanceStack.AvailableInstances)
 		{
-			Instance->GetComponent()->SetVisibility(bVisible);
+			if (USceneComponent* Component = Instance->GetComponent())
+			{
+				Component->SetVisibility(bVisible);
+			}
 		}
 
 		for (FAdditiveAnimationInstance* Instance : Data->AdditiveInstanceStack.InUseInstances)
 		{
-			Instance->GetComponent()->SetVisibility(bVisible);
+			if (USceneComponent* Component = Instance->GetComponent())
+			{
+				Component->SetVisibility(bVisible);
+			}
 		}
 	}
 }
@@ -746,7 +771,7 @@ void UAnimSharingInstance::Setup(UAnimationSharingManager* AnimationSharingManag
 			const uint8 StateValue = StateEntry.State;
 			const uint32 StateIndex = StateEnum->GetIndexByValue(StateValue);
 
-			checkf(!PerStateData.FindByPredicate([&](const FPerStateData& State) { return State.StateEnumValue == StateValue; }), TEXT("State already defined"));
+			checkf(!PerStateData.FindByPredicate([StateValue](const FPerStateData& State) { return State.StateEnumValue == StateValue; }), TEXT("State already defined"));
 
 			FPerStateData& StateData = PerStateData[StateIndex];
 			StateData.StateEnumValue = StateValue;
@@ -997,7 +1022,7 @@ void UAnimSharingInstance::TickDebugInformation()
 			const FPerActorData& ActorData = PerActorData[ActorIndex];
 			const uint8 State = ActorData.CurrentState;
 
-			const FString StateString = [&]() -> FString
+			const FString StateString = [this, &ActorData, State]() -> FString
 			{
 				/** Check whether or not we are currently blending */
 				const uint32 BlendInstanceIndex = ActorData.BlendInstanceIndex;
@@ -1018,7 +1043,7 @@ void UAnimSharingInstance::TickDebugInformation()
 				return FString::Printf(TEXT("State - %s %1.2f"), *StateEnum->GetDisplayNameTextByValue(State).ToString(), ActorData.SignificanceValue);
 			}();
 
-			const FColor DebugColor = [&]()
+			const FColor DebugColor = [&ActorData]()
 			{
 				const uint32 BlendInstanceIndex = ActorData.BlendInstanceIndex;
 				const uint32 DemandInstanceIndex = ActorData.OnDemandInstanceIndex;
@@ -1088,7 +1113,7 @@ void UAnimSharingInstance::TickOnDemandInstances()
 #if LOG_STATES 
 			UE_LOG(LogAnimationSharing, Log, TEXT("Finished on demand %s"), *StateEnum->GetDisplayNameTextByValue(Instance.State).ToString());
 #endif
-			auto SetActorState = [&](uint32 ActorIndex, uint8 NewState)
+			auto SetActorState = [this, &Instance](uint32 ActorIndex, uint8 NewState)
 			{
 				if (Instance.BlendToPermutationIndex != INDEX_NONE)
 				{
@@ -1177,7 +1202,7 @@ void UAnimSharingInstance::TickOnDemandInstances()
 				const bool bShouldBlend = ScalabilitySettings->UseBlendTransitions.Default && PerActorData[ActorIndex].SignificanceValue >= ScalabilitySettings->BlendSignificanceValue.Default;
 
 				// Determine state to blend to
-				const uint8 BlendToState = [&]() -> uint8
+				const uint8 BlendToState = [this, &Instance, bShouldBlend, ActorIndex]() -> uint8
 				{
 					if (bShouldBlend)
 					{
@@ -1302,7 +1327,7 @@ void UAnimSharingInstance::TickActorStates()
 				/** When we are currently running an on-demand state we do not want as state change to impact the current animation */
 				const bool bShouldNotProcess = ActorData.bRunningOnDemand && !PerStateData[CurrentState].bIsOnDemand;
 
-				auto UpdateState = [&]()
+				auto UpdateState = [&ActorData, ActorIndex, CurrentState, PreviousState]()
 				{
 #if LOG_STATES 
 					UE_LOG(LogAnimationSharing, Log, TEXT("Setting %i state to %i previous %i | %i"), ActorIndex, CurrentState, PreviousState, ActorData.PermutationIndex);
@@ -1736,7 +1761,7 @@ uint32 UAnimSharingInstance::SetupBlend(uint8 FromState, uint8 ToState, uint32 A
 	uint32 BlendInstanceIndex = INDEX_NONE;
 	if (!bConcurrentBlendsReached)
 	{
-		BlendInstanceIndex = BlendInstances.IndexOfByPredicate([&](const FBlendInstance& Instance)
+		BlendInstanceIndex = BlendInstances.IndexOfByPredicate([this, FromState, ToState, bOnDemand, ActorIndex](const FBlendInstance& Instance)
 		{			
 			return (!Instance.bActive &&				// The instance should not have started yet
 				Instance.StateFrom == FromState &&		// It should be blending from the same state
diff --git a/Engine/Plugins/Developer/AnimationSharing/Source/AnimationSharing/Private/TransitionBlendInstance.cpp b/Engine/Plugins/Developer/AnimationSharing/Source/AnimationSharing/Private/TransitionBlendInstance.cpp
--- a/Engine/Plugins/Developer/AnimationSharing/Source/AnimationSharing/Private/TransitionBlendInstance.cpp
+++ b/Engine/Plugins/Developer/AnimationSharing/Source/AnimationSharing/Private/TransitionBlendInstance.cpp
@@ -8,14 +8,18 @@ FTransitionBlendInstance::FTransitionBlendInstance() : SkeletalMeshComponent(nul
 
 void FTransitionBlendInstance::Initialise(USkeletalMeshComponent* InSkeletalMeshComponent, UClass* InAnimationBPClass)
 {
-	if (InSkeletalMeshComponent && InAnimationBPClass)
+	if (InSkeletalMeshComponent)
 	{
 		SkeletalMeshComponent = InSkeletalMeshComponent;
-		SkeletalMeshComponent->SetAnimInstanceClass(InAnimationBPClass);
+		if (InAnimationBPClass)
+		{
+			SkeletalMeshComponent->SetAnimInstanceClass(InAnimationBPClass);
+			TransitionInstance = Cast<UAnimSharingTransitionInstance>(SkeletalMeshComponent->GetAnimInstance());
+		}
+		
 		SkeletalMeshComponent->SetComponentTickEnabled(false);	
 		SkeletalMeshComponent->SetForcedLOD(0);
-		SkeletalMeshComponent->VisibilityBasedAnimTickOption = EVisibilityBasedAnimTickOption::AlwaysTickPoseAndRefreshBones;
-		TransitionInstance = Cast<UAnimSharingTransitionInstance>(SkeletalMeshComponent->GetAnimInstance());
+		SkeletalMeshComponent->VisibilityBasedAnimTickOption = EVisibilityBasedAnimTickOption::AlwaysTickPoseAndRefreshBones;		
 	}
 }
 
diff --git a/Engine/Plugins/Developer/AnimationSharing/Source/AnimationSharingEd/Private/SetupDetailsViewCustomizations.cpp b/Engine/Plugins/Developer/AnimationSharing/Source/AnimationSharingEd/Private/SetupDetailsViewCustomizations.cpp
--- a/Engine/Plugins/Developer/AnimationSharing/Source/AnimationSharingEd/Private/SetupDetailsViewCustomizations.cpp
+++ b/Engine/Plugins/Developer/AnimationSharing/Source/AnimationSharingEd/Private/SetupDetailsViewCustomizations.cpp
@@ -6,6 +6,7 @@
 #include "UObject/UObjectGlobals.h"
 #include "Widgets/Input/STextComboBox.h"
 #include "AnimationSharingTypes.h"
+#include "Templates/SharedPointer.h"
 
 #define LOCTEXT_NAMESPACE "AnimationSharingSetupCustomization"
 
@@ -337,12 +338,48 @@ UEnum* GetStateEnumClass(const TSharedPtr<IPropertyHandle>& InProperty)
 	return nullptr;
 }
 
-
 TSharedRef<IPropertyTypeCustomization> FAnimationSetupCustomization::MakeInstance()
 {
 	return MakeShareable(new FAnimationSetupCustomization());
 }
 
+void FAnimationSetupCustomization::CustomizeHeader(TSharedRef<IPropertyHandle> PropertyHandle, FDetailWidgetRow& HeaderRow, IPropertyTypeCustomizationUtils& CustomizationUtils)
+{
+	const FName AnimSequencePropertyName = GET_MEMBER_NAME_CHECKED(FAnimationSetup, AnimSequence);
+	AnimSequencePropertyHandle = PropertyHandle->GetChildHandle(AnimSequencePropertyName);
+	
+	HeaderRow
+	.NameContent()
+	[
+		SNew(SHorizontalBox)
+		+SHorizontalBox::Slot()
+		.FillWidth(1)
+		.VAlign(VAlign_Center)
+		[
+			// Show the name of the asset or actor
+			SNew(STextBlock)
+			.Font(FEditorStyle::GetFontStyle("PropertyWindow.NormalFont"))
+			.Text_Lambda([this]() -> FText
+			{
+				if (AnimSequencePropertyHandle.IsValid())
+				{
+					FText PropertyValueAsName;
+					if (AnimSequencePropertyHandle->GetValueAsFormattedText(PropertyValueAsName) == FPropertyAccess::Success)
+					{
+						return PropertyValueAsName;
+					}
+				}
+
+				return LOCTEXT("None", "None");
+			})			
+		]
+	];
+	/*.ValueContent()
+	[
+		PropertyHandle->CreatePropertyValueWidget()
+	];*/	
+}
+
 void FAnimationSetupCustomization::CustomizeChildren(TSharedRef<IPropertyHandle> StructPropertyHandle, IDetailChildrenBuilder& StructBuilder, IPropertyTypeCustomizationUtils& StructCustomizationUtils)
 {
 	uint32 NumChildren;
diff --git a/Engine/Plugins/Developer/AnimationSharing/Source/AnimationSharingEd/Public/SetupDetailsViewCustomizations.h b/Engine/Plugins/Developer/AnimationSharing/Source/AnimationSharingEd/Public/SetupDetailsViewCustomizations.h
--- a/Engine/Plugins/Developer/AnimationSharing/Source/AnimationSharingEd/Public/SetupDetailsViewCustomizations.h
+++ b/Engine/Plugins/Developer/AnimationSharing/Source/AnimationSharingEd/Public/SetupDetailsViewCustomizations.h
@@ -65,10 +65,13 @@ public:
 	virtual ~FAnimationSetupCustomization() {}
 
 	/** Begin IPropertyTypeCustomization interface */
-	virtual void CustomizeHeader(TSharedRef<IPropertyHandle> PropertyHandle, FDetailWidgetRow& HeaderRow, IPropertyTypeCustomizationUtils& CustomizationUtils) override {}
+	virtual void CustomizeHeader(TSharedRef<IPropertyHandle> PropertyHandle, FDetailWidgetRow& HeaderRow, IPropertyTypeCustomizationUtils& CustomizationUtils) override;
 	virtual void CustomizeChildren(TSharedRef<IPropertyHandle> StructPropertyHandle, IDetailChildrenBuilder& StructBuilder, IPropertyTypeCustomizationUtils& StructCustomizationUtils) override;
 	/** End IPropertyTypeCustomization interface */
 
+protected:
+	FText GetAnimationName(TSharedPtr<IPropertyHandle> PropertyHandle) const;
+
 protected:
 	TSharedPtr<IPropertyHandle> AnimSequencePropertyHandle;
 };
diff --git a/Engine/Plugins/Experimental/GeometryCache/Source/GeometryCache/Private/GeometryCacheSceneProxy.cpp b/Engine/Plugins/Experimental/GeometryCache/Source/GeometryCache/Private/GeometryCacheSceneProxy.cpp
--- a/Engine/Plugins/Experimental/GeometryCache/Source/GeometryCache/Private/GeometryCacheSceneProxy.cpp
+++ b/Engine/Plugins/Experimental/GeometryCache/Source/GeometryCache/Private/GeometryCacheSceneProxy.cpp
@@ -727,6 +727,7 @@ FPrimitiveViewRelevance FGeometryCacheSceneProxy::GetViewRelevance(const FSceneV
 	Result.bDrawRelevance = IsShown(View);
 	Result.bShadowRelevance = IsShadowCast(View);
 	Result.bDynamicRelevance = true;
+	Result.bRenderCustomDepth = ShouldRenderCustomDepth();
 	MaterialRelevance.SetPrimitiveViewRelevance(Result);
 	return Result;
 }
diff --git a/Engine/Plugins/Experimental/GeometryCache/Source/GeometryCacheTracks/Private/MovieSceneGeometryCacheSection.cpp b/Engine/Plugins/Experimental/GeometryCache/Source/GeometryCacheTracks/Private/MovieSceneGeometryCacheSection.cpp
--- a/Engine/Plugins/Experimental/GeometryCache/Source/GeometryCacheTracks/Private/MovieSceneGeometryCacheSection.cpp
+++ b/Engine/Plugins/Experimental/GeometryCache/Source/GeometryCacheTracks/Private/MovieSceneGeometryCacheSection.cpp
@@ -42,6 +42,8 @@ TOptional<FFrameTime> UMovieSceneGeometryCacheSection::GetOffsetTime() const
 
 void UMovieSceneGeometryCacheSection::PostLoad()
 {
+	Super::PostLoad();
+
 	FFrameRate LegacyFrameRate = GetLegacyConversionFrameRate();
 
 	if (Params.StartOffset_DEPRECATED != GeometryCacheDeprecatedMagicNumber)
@@ -57,7 +59,6 @@ void UMovieSceneGeometryCacheSection::PostLoad()
 
 		Params.EndOffset_DEPRECATED = GeometryCacheDeprecatedMagicNumber;
 	}
-
 }
 
 void UMovieSceneGeometryCacheSection::Serialize(FArchive& Ar)
diff --git a/Engine/Plugins/Experimental/ProxyLODPlugin/Source/ProxyLOD/Private/ProxyLODPlugin.cpp b/Engine/Plugins/Experimental/ProxyLODPlugin/Source/ProxyLOD/Private/ProxyLODPlugin.cpp
--- a/Engine/Plugins/Experimental/ProxyLODPlugin/Source/ProxyLOD/Private/ProxyLODPlugin.cpp
+++ b/Engine/Plugins/Experimental/ProxyLODPlugin/Source/ProxyLOD/Private/ProxyLODPlugin.cpp
@@ -929,9 +929,13 @@ void FVoxelizeMeshMerging::ProxyLOD(const FMeshMergeDataArray& InData, const FMe
 				{
 					const auto& FirstMaterial = (*BakedMaterials)[0];
 					OutMaterial.BlendMode              = FirstMaterial.BlendMode;
-					OutMaterial.bTwoSided              = FirstMaterial.bTwoSided;
-					OutMaterial.bDitheredLODTransition = FirstMaterial.bDitheredLODTransition;
-					OutMaterial.EmissiveScale          = FirstMaterial.EmissiveScale;
+
+					for (const FFlattenMaterial& FlatMaterial : (*BakedMaterials))
+					{
+						OutMaterial.bTwoSided |= FlatMaterial.bTwoSided;
+						OutMaterial.bDitheredLODTransition |= FlatMaterial.bDitheredLODTransition;
+						OutMaterial.EmissiveScale = FMath::Max(OutMaterial.EmissiveScale, FlatMaterial.EmissiveScale);
+					}
 				}
 
 				// --- Map a correspondence between texels in the texture atlas and locations on the source geometry --
diff --git a/Engine/Source/Developer/MaterialBaking/Private/MaterialBakingHelpers.cpp b/Engine/Source/Developer/MaterialBaking/Private/MaterialBakingHelpers.cpp
--- a/Engine/Source/Developer/MaterialBaking/Private/MaterialBakingHelpers.cpp
+++ b/Engine/Source/Developer/MaterialBaking/Private/MaterialBakingHelpers.cpp
@@ -4,35 +4,26 @@
 #include "Math/Color.h"
 #include "Async/ParallelFor.h"
 
-static FColor BoxBlurSample(TArray<FColor>& InBMP, int32 X, int32 Y, int32 InImageWidth, int32 InImageHeight)
+static FColor BoxBlurSample(FColor* InBMP, int32 X, int32 Y, int32 InImageWidth, int32 InImageHeight)
 {
-	const int32 SampleCount = 8;
-	const int32 PixelIndices[SampleCount] = { -(InImageWidth + 1), -InImageWidth, -(InImageWidth - 1),
-		-1, 1,
-		(InImageWidth + -1), InImageWidth, (InImageWidth + 1) };
-
-	static const int32 PixelOffsetX[SampleCount] = { -1, 0, 1,
-		-1, 1,
-		-1, 0, 1 };
-
 	int32 PixelsSampled = 0;
 	int32 CombinedColorR = 0;
 	int32 CombinedColorG = 0;
 	int32 CombinedColorB = 0;
 	int32 CombinedColorA = 0;
+	
+	const uint32 MagentaMask = FColor(255, 0, 255).DWColor();
 
-	// Take samples for blur with square indices
-	for (int32 SampleIndex = 0; SampleIndex < SampleCount; ++SampleIndex)
+	const int32 BaseIndex = (Y * InImageWidth) + X;
+	
+	int32 SampleIndex = BaseIndex - (InImageWidth + 1);
+	for (int32 YI = 0; YI < 3; ++YI)
 	{
-		const int32 PixelIndex = ((Y * InImageWidth) + X) + PixelIndices[SampleIndex];
-		const int32 XIndex = X + PixelOffsetX[SampleIndex];
-
-		// Check if we are not out of texture bounds
-		if (InBMP.IsValidIndex(PixelIndex) && XIndex >= 0 && XIndex < InImageWidth)
+		for (int32 XI = 0; XI < 3; ++XI)
 		{
-			const FColor SampledColor = InBMP[PixelIndex];
+			const FColor SampledColor = InBMP[SampleIndex++];
 			// Check if the pixel is a rendered one (not clear color)
-			if ((!(SampledColor.R == 255 && SampledColor.B == 255 && SampledColor.G == 0)))
+			if (SampledColor.DWColor() != MagentaMask)
 			{
 				CombinedColorR += SampledColor.R;
 				CombinedColorG += SampledColor.G;
@@ -41,112 +32,239 @@ static FColor BoxBlurSample(TArray<FColor>& InBMP, int32 X, int32 Y, int32 InIma
 				++PixelsSampled;
 			}
 		}
+		SampleIndex += InImageWidth - 3;
 	}
 
 	if (PixelsSampled == 0)
 	{
-		return InBMP[((Y * InImageWidth) + X)];
+		const FColor Magenta = FColor(255, 0, 255);
+		return Magenta;
 	}
 
 	return FColor(CombinedColorR / PixelsSampled, CombinedColorG / PixelsSampled, CombinedColorB / PixelsSampled, CombinedColorA / PixelsSampled);
 }
 
+static bool HasBorderingPixel(FColor* InBMP, int32 X, int32 Y, int32 InImageWidth, int32 InImageHeight)
+{
+	const uint32 MagentaMask = FColor(255, 0, 255).DWColor();
+
+	const int32 BaseIndex = Y*InImageWidth + X;
+
+	int32 SampleIndex = BaseIndex - (InImageWidth + 1);
+	for (int32 YI = 0; YI < 3; ++YI)
+	{
+		for (int32 XI = 0; XI < 3; ++XI)
+		{
+			if (InBMP[SampleIndex++].DWColor() != MagentaMask)
+			{
+				return true;
+			}
+		}
+		SampleIndex += InImageWidth - 3;
+	}
+	return false;
+}
+
 void FMaterialBakingHelpers::PerformUVBorderSmear(TArray<FColor>& InOutPixels, int32 ImageWidth, int32 ImageHeight, int32 MaxIterations)
 {
-	// This is ONLY possible because this is never called from multiple threads
-	static TArray<FColor> Swap;
-	if (Swap.Num())
+	check(InOutPixels.Num() == ImageWidth*ImageHeight);
+
+	const uint32 MagentaMask = FColor(255, 0, 255).DWColor();
+	const int32 PaddedImageWidth = ImageWidth + 2;
+	const int32 PaddedImageHeight = ImageHeight + 2;
+
+	TArray<uint32> RowsSwap1;
+	TArray<uint32> RowsSwap2;
+	TArray<uint32>* CurrentRowsLeft = &RowsSwap1;
+	TArray<uint32>* NextRowsLeft = &RowsSwap2;
+
+	//
+	// Create a copy of the image with a 1 pixel magenta border around the edge.
+	// This avoids needing to bounds check in inner loops.
+	//
+
+	FColor* Current = InOutPixels.GetData();
+	
+	TArray<FColor> ScratchBuffer;
+	ScratchBuffer.SetNumUninitialized(PaddedImageWidth*PaddedImageHeight);
+	FColor* Scratch = ScratchBuffer.GetData();
+
+	// Set top and bottom bordering pixels to magenta
+	for (int32 X = 0; X < PaddedImageWidth; ++X)
+	{
+		Scratch[X] = FColor(MagentaMask);
+		Scratch[(PaddedImageHeight - 1)*PaddedImageWidth + X] = FColor(MagentaMask);
+	}
+
+	// Set leftg and right border pixels to magenta and copy image data into rows
+	for (int32 Y = 1; Y <= ImageHeight; ++Y)
 	{
-		Swap.SetNum(0, false);
+		const int32 YOffset = Y * PaddedImageWidth;
+		Scratch[YOffset] = FColor(MagentaMask);
+		FMemory::Memcpy(&Scratch[YOffset + 1], &Current[(Y - 1)*ImageWidth], ImageWidth * sizeof(FColor));
+		Scratch[YOffset + ImageWidth + 1] = FColor(MagentaMask);
 	}
-	Swap.Append(InOutPixels);
 
-	TArray<FColor>* Current = &InOutPixels;
-	TArray<FColor>* Scratch = &Swap;
+	//
+	// Find our initial workset of all rows that have magenta pixels bordering non-magenta pixels.
+	// Also find all rows that have zero magenta pixels - these will never be added to the workset.
+	//
+
+	TArray<bool> RowCompleted;
+	RowCompleted.SetNumZeroed(PaddedImageHeight);
 
-	MaxIterations = MaxIterations < 1 ? FMath::Max(ImageWidth, ImageHeight) : MaxIterations;
-	const int32 NumThreads = [&]()
+	// Set border rows to start completed
+	RowCompleted[0] = true;
+	RowCompleted[PaddedImageHeight - 1] = true;
+
+	bool bHasAnyData = false;
+	for (int32 Y = 1; Y <= ImageHeight; ++Y)
 	{
-		return FPlatformProcess::SupportsMultithreading() ? FPlatformMisc::NumberOfCores() : 1;
-	}();
+		bool bHasMagenta = false;
+		bool bBordersNonMagenta = false;
+		const int32 YOffset = Y * PaddedImageWidth;
+		for (int32 X = 1; X <= ImageWidth; ++X)
+		{
+			if (Scratch[YOffset + X].DWColor() == MagentaMask)
+			{
+				bHasMagenta = true;
+				if (HasBorderingPixel(Scratch, X, Y, PaddedImageWidth, PaddedImageHeight))
+				{
+					bBordersNonMagenta = true;
+					break;
+				}
+			}
+		}
 
-	const int32 LinesPerThread = FMath::CeilToInt((float)ImageHeight / (float)NumThreads);
-	uint32* MagentaPixels = new uint32[NumThreads];
-	FMemory::Memset(MagentaPixels, 0xff, sizeof(uint32) * NumThreads);
+		if (!bHasMagenta)
+		{
+			bHasAnyData = true;
+			RowCompleted[Y] = true;
+		}
+		else if (bBordersNonMagenta)
+		{
+			bHasAnyData = true;
+			CurrentRowsLeft->Add(Y);
+		}
+	}
+
+	// early-out if no valid pixels were found.
+	if (!bHasAnyData)
+	{
+		FMemory::Memzero(InOutPixels.GetData(), ImageWidth * ImageHeight * sizeof(FColor));
+		return;
+	}
+
+	// Early out if all pixels were valid.
+	if (CurrentRowsLeft->Num() == 0)
+	{
+		return;
+	}
 
-	uint32 TotalPixels = ImageWidth * ImageHeight;
-	uint32 SummedMagentaPixels = 1;
+	TArray<uint32> RowRemainingPixels;
+	RowRemainingPixels.SetNumZeroed(PaddedImageHeight);
+
+	const int32 MaxThreads = FPlatformProcess::SupportsMultithreading() ? FPlatformMisc::NumberOfCores() : 1;
+
+	//
+	// Iteratively smear until all rows are filled.
+	//
 
-	// Sampling
 	int32 LoopCount = 0;
-	while (SummedMagentaPixels && (LoopCount <= MaxIterations))
+	while (CurrentRowsLeft->Num() && (MaxIterations <= 0 || LoopCount <= MaxIterations))
 	{
-		SummedMagentaPixels = 0;
-		// Left / right, Top / down
-		ParallelFor(NumThreads, [ImageWidth, ImageHeight, MaxIterations, LoopCount, Current, Scratch, &MagentaPixels, LinesPerThread]
-		(int32 Index)
+		const int32 NumThreads = FMath::Min(CurrentRowsLeft->Num(), MaxThreads);
+		const int32 LinesPerThread = FMath::CeilToInt((float)CurrentRowsLeft->Num() / (float)NumThreads);
+
+		// split up rows that still have magenta pixels amongst threads
+		ParallelFor(NumThreads, [ImageWidth, PaddedImageWidth, PaddedImageHeight, Current, Scratch, CurrentRowsLeft, LinesPerThread, &RowRemainingPixels, MagentaMask](int32 ThreadIndex)
 		{
-			if(MagentaPixels[Index] > 0)
-			{
-				MagentaPixels[Index] = 0;
+			const int32 StartY = ThreadIndex*LinesPerThread;
+			const int32 EndY = FMath::Min((ThreadIndex + 1) * LinesPerThread, CurrentRowsLeft->Num());
 
-				const int32 StartY = FMath::CeilToInt((Index)* LinesPerThread);
-				const int32 EndY = FMath::Min(FMath::CeilToInt((Index + 1) * LinesPerThread), ImageHeight);
+			for (int32 Index = StartY; Index < EndY; Index++)
+			{
+				const int32 Y = (*CurrentRowsLeft)[Index];
+				RowRemainingPixels[Index] = 0;
 
-				for (int32 Y = StartY; Y < EndY; Y++)
+				int32 PixelIndex = (Y - 1) * ImageWidth;
+				for (int32 X = 1; X <= ImageWidth; X++)
 				{
-					for (int32 X = 0; X < ImageWidth; X++)
+					FColor& Color = Current[PixelIndex++];
+					if (Color.DWColor() == MagentaMask)
 					{
-						const int32 PixelIndex = (Y * ImageWidth) + X;
-						FColor& Color = (*Current)[PixelIndex];
-						if (Color.R == 255 && Color.B == 255 && Color.G == 0)
+						const FColor SampledColor = BoxBlurSample(Scratch, X, Y, PaddedImageWidth, PaddedImageHeight);
+						// If it's a valid pixel
+						if (SampledColor.DWColor() != MagentaMask)
+						{
+							Color = SampledColor;
+						}
+						else
 						{
-							MagentaPixels[Index]++;
-							const FColor SampledColor = BoxBlurSample(*Scratch, X, Y, ImageWidth, ImageHeight);
-							// If it's a valid pixel
-							if (!(SampledColor.R == 255 && SampledColor.B == 255 && SampledColor.G == 0))
-							{
-								Color = SampledColor;
-							}
-							else
-							{
-								// If we are at the end of our iterations, replace the pixels with black
-								if (LoopCount == (MaxIterations - 1))
-								{
-									Color = FColor(0, 0, 0, 0);
-								}
-							}
+							RowRemainingPixels[Index]++;
 						}
 					}
 				}
 			}
 		});
 
-		for (int32 ThreadIndex = 0; ThreadIndex < NumThreads; ++ThreadIndex)
+		// Mark all completed rows and copy data between buffers
+		const int32 RowMemorySize = ImageWidth * sizeof(FColor);
+		for (int32 RowIndex = 0; RowIndex < CurrentRowsLeft->Num(); ++RowIndex)
 		{
-			SummedMagentaPixels += MagentaPixels[ThreadIndex];
+			const int32 Y = (*CurrentRowsLeft)[RowIndex];
+			if (!RowRemainingPixels[RowIndex])
+			{
+				RowCompleted[Y] = true;
+			}
+			FMemory::Memcpy(&Scratch[Y * PaddedImageWidth + 1], &Current[(Y - 1) * ImageWidth], RowMemorySize);
 		}
 
-		// If after iterating over all pixels we found they were all magenta, exit early
-		if(SummedMagentaPixels >= TotalPixels)
+		// Find the next set of rows to process.
+		// This will be our current rows that haven't completed, plus any bordering rows that haven't been marked completed.
+		NextRowsLeft->SetNum(0, false);
 		{
-			FMemory::Memset(Current->GetData(), 0, Current->Num() * sizeof(FColor));
-			break;
+			int32 PreviousY = -1;
+			for (int32 RowIndex = 0; RowIndex < CurrentRowsLeft->Num(); ++RowIndex)
+			{
+				const int32 Y = (*CurrentRowsLeft)[RowIndex];
+				if (!RowCompleted[Y - 1] && PreviousY < Y - 1)
+				{
+					PreviousY = Y - 1;
+					NextRowsLeft->Add(Y - 1);
+				}
+				if (!RowCompleted[Y] && RowRemainingPixels[RowIndex] && PreviousY < Y)
+				{
+					PreviousY = Y;
+					NextRowsLeft->Add(Y);
+				}
+				if (!RowCompleted[Y + 1])
+				{
+					PreviousY = Y + 1;
+					NextRowsLeft->Add(Y + 1);
+				}
+			}
 		}
 
-		TArray<FColor>& Temp = *Scratch;
-		Scratch = Current;
-		Current = &Temp;
+		// Swap rows left buffers
+		{			
+			TArray<uint32>* Temp = NextRowsLeft;
+			NextRowsLeft = CurrentRowsLeft;
+			CurrentRowsLeft = Temp;
+		}
 
 		LoopCount++;
 	}
 
-	if (Current != &InOutPixels)
+	// If we finished before replacing all pixels, replace the remaining pixels with black.
+	if (CurrentRowsLeft->Num() > 0)
 	{
-		InOutPixels.Empty();
-		InOutPixels.Append(*Current);
+		for (int32 i = 0; i < InOutPixels.Num(); ++i)
+		{
+			if (InOutPixels[i].DWColor() == MagentaMask)
+			{
+				InOutPixels[i] = FColor::Black;
+			}
+		}
 	}
-
-	delete[] MagentaPixels;
 }
-
diff --git a/Engine/Source/Developer/MaterialBaking/Private/MaterialBakingModule.cpp b/Engine/Source/Developer/MaterialBaking/Private/MaterialBakingModule.cpp
--- a/Engine/Source/Developer/MaterialBaking/Private/MaterialBakingModule.cpp
+++ b/Engine/Source/Developer/MaterialBaking/Private/MaterialBakingModule.cpp
@@ -199,12 +199,12 @@ void FMaterialBakingModule::BakeMaterials(const TArray<FMaterialData*>& Material
 					FlushRenderingCommands();
 
 					SortElement.RenderBatchArray.Empty();
-					ReadTextureOutput(RenderTargetResource, Property, CurrentOutput);
-
-					if(CurrentMaterialSettings->bPerformBorderSmear)
-					{
-						FMaterialBakingHelpers::PerformUVBorderSmear(CurrentOutput.PropertyData[Property], RenderTarget->GetSurfaceWidth(), RenderTarget->GetSurfaceHeight());
-					}
+						ReadTextureOutput(RenderTargetResource, Property, CurrentOutput);
+						
+						if (CurrentMaterialSettings->bPerformBorderSmear)
+						{
+							FMaterialBakingHelpers::PerformUVBorderSmear(CurrentOutput.PropertyData[Property], RenderTarget->GetSurfaceWidth(), RenderTarget->GetSurfaceHeight());
+						}
 #if WITH_EDITOR
 					// If saving intermediates is turned on
 					if (bSaveIntermediateTextures)
@@ -323,7 +323,7 @@ FExportMaterialProxy* FMaterialBakingModule::CreateMaterialProxy(UMaterialInterf
 
 void FMaterialBakingModule::ReadTextureOutput(FTextureRenderTargetResource* RenderTargetResource, EMaterialProperty Property, FBakeOutput& Output)
 {	
-	checkf(!Output.PropertyData.Contains(Property) && !Output.PropertySizes.Contains(Property), TEXT("Should be reading same property data twice"));
+	checkf(!Output.PropertyData.Contains(Property) && !Output.PropertySizes.Contains(Property), TEXT("Should not be reading same property data twice"));
 
 	TArray<FColor>& OutputColor = Output.PropertyData.Add(Property);
 	FIntPoint& OutputSize = Output.PropertySizes.Add(Property);	
@@ -404,10 +404,23 @@ void FMaterialBakingModule::ReadTextureOutput(FTextureRenderTargetResource* Rend
 				for (int32 PixelX = 0; PixelX < OutputSize.X; PixelX++)
 				{
 					const FFloat16Color& Pixel16 = Color16[PixelX + YOffset];
+
+					const FFloat16Color MagentaFloat16 = FFloat16Color(FLinearColor(1.0f, 0.0f, 1.0f));
 					FColor& Pixel8 = OutputColor[PixelX + YOffset];
-					Pixel8.R = (uint8)FMath::RoundToInt(Pixel16.R.GetFloat() * Scale);
-					Pixel8.G = (uint8)FMath::RoundToInt(Pixel16.G.GetFloat() * Scale);
-					Pixel8.B = (uint8)FMath::RoundToInt(Pixel16.B.GetFloat() * Scale);
+					/*if (Pixel16 == MagentaFloat16)
+					{
+						Pixel8.R = 255;
+						Pixel8.G = 0;
+						Pixel8.B = 255;
+					}
+					else*/
+					{
+						Pixel8.R = (uint8)FMath::RoundToInt(Pixel16.R.GetFloat() * Scale);
+						Pixel8.G = (uint8)FMath::RoundToInt(Pixel16.G.GetFloat() * Scale);
+						Pixel8.B = (uint8)FMath::RoundToInt(Pixel16.B.GetFloat() * Scale);						
+					}
+					
+					
 					Pixel8.A = 255;
 				}
 			}
diff --git a/Engine/Source/Developer/MeshMergeUtilities/Private/MeshMergeHelpers.cpp b/Engine/Source/Developer/MeshMergeUtilities/Private/MeshMergeHelpers.cpp
--- a/Engine/Source/Developer/MeshMergeUtilities/Private/MeshMergeHelpers.cpp
+++ b/Engine/Source/Developer/MeshMergeUtilities/Private/MeshMergeHelpers.cpp
@@ -911,13 +911,13 @@ void FMeshMergeHelpers::RetrieveCullingLandscapeAndVolumes(UWorld* InWorld, cons
 	for (ALandscapeProxy* Landscape : LandscapeActors)
 	{
 		// Export the landscape to raw mesh format
-		FMeshDescription MeshDescription;
-		UStaticMesh::RegisterMeshAttributes(MeshDescription);
+		FMeshDescription* MeshDescription = new FMeshDescription();
+		UStaticMesh::RegisterMeshAttributes(*MeshDescription);
 		FBoxSphereBounds LandscapeBounds = EstimatedMeshProxyBounds;
-		Landscape->ExportToRawMesh(LandscapeExportLOD, MeshDescription, LandscapeBounds);
-		if (MeshDescription.Vertices().Num())
+		Landscape->ExportToRawMesh(LandscapeExportLOD, *MeshDescription, LandscapeBounds);
+		if (MeshDescription->Vertices().Num())
 		{
-			CullingRawMeshes.Add(&MeshDescription);
+			CullingRawMeshes.Add(MeshDescription);
 		}
 	}
 
diff --git a/Engine/Source/Developer/MeshMergeUtilities/Private/MeshMergeHelpers.h b/Engine/Source/Developer/MeshMergeUtilities/Private/MeshMergeHelpers.h
--- a/Engine/Source/Developer/MeshMergeUtilities/Private/MeshMergeHelpers.h
+++ b/Engine/Source/Developer/MeshMergeUtilities/Private/MeshMergeHelpers.h
@@ -54,7 +54,7 @@ public:
 	/** Transforms raw mesh data using InTransform*/
 	static void TransformRawMeshVertexData(const FTransform& InTransform, FMeshDescription &OutRawMesh);
 
-	/** Retrieves all culling landscapes and volumes as FMeshDescription structures */
+	/** Retrieves all culling landscapes and volumes as FMeshDescription structures. Note the caller is responsible for deleting the heap data managed by CullingRawMeshes */
 	static void RetrieveCullingLandscapeAndVolumes(UWorld* InWorld, const FBoxSphereBounds& EstimatedMeshProxyBounds, const TEnumAsByte<ELandscapeCullingPrecision::Type> PrecisionType, TArray<FMeshDescription*>& CullingRawMeshes);
 
 	/** Transforms physics geometry data using InTransform */
diff --git a/Engine/Source/Developer/MeshMergeUtilities/Private/MeshMergeUtilities.cpp b/Engine/Source/Developer/MeshMergeUtilities/Private/MeshMergeUtilities.cpp
--- a/Engine/Source/Developer/MeshMergeUtilities/Private/MeshMergeUtilities.cpp
+++ b/Engine/Source/Developer/MeshMergeUtilities/Private/MeshMergeUtilities.cpp
@@ -1484,7 +1484,7 @@ void FMeshMergeUtilities::CreateProxyMesh(const TArray<UStaticMeshComponent*>& I
 		}
 	};
 
-	// Landscape culling
+	// Landscape culling.  NB these are temporary copies of the culling data and should be deleted after use.
 	TArray<FMeshDescription*> CullingRawMeshes;
 	if (InMeshProxySettings.bUseLandscapeCulling)
 	{
@@ -1587,8 +1587,15 @@ void FMeshMergeUtilities::CreateProxyMesh(const TArray<UStaticMeshComponent*>& I
 
 		Processor->Tick(0); // make sure caller gets merging results
 	}
+
+	// Clean up the CullingRawMeshes
+	for (FMeshDescription* RawMesh : CullingRawMeshes)
+	{
+		delete RawMesh;
+	}
 }
 
+
 bool FMeshMergeUtilities::IsValidBaseMaterial(const UMaterialInterface* InBaseMaterial, bool bShowToaster) const
 {
 	if (InBaseMaterial != nullptr)
@@ -3077,12 +3084,12 @@ UMaterialInterface* FMeshMergeUtilities::CreateProxyMaterial(const FString &InBa
 	FString MaterialPackageName;
 	if (InBasePackageName.IsEmpty())
 	{
-		MaterialAssetName = TEXT("M_MERGED_") + FPackageName::GetShortName(MergedAssetPackageName);
+		MaterialAssetName = FPackageName::GetShortName(MergedAssetPackageName);
 		MaterialPackageName = FPackageName::GetLongPackagePath(MergedAssetPackageName) + TEXT("/") + MaterialAssetName;
 	}
 	else
 	{
-		MaterialAssetName = TEXT("M_") + FPackageName::GetShortName(InBasePackageName);
+		MaterialAssetName = FPackageName::GetShortName(InBasePackageName);
 		MaterialPackageName = FPackageName::GetLongPackagePath(InBasePackageName) + TEXT("/") + MaterialAssetName;
 	}
 
@@ -3095,7 +3102,7 @@ UMaterialInterface* FMeshMergeUtilities::CreateProxyMaterial(const FString &InBa
 		MaterialPackage->Modify();
 	}
 
-	UMaterialInstanceConstant* MergedMaterial = ProxyMaterialUtilities::CreateProxyMaterialInstance(MaterialPackage, InSettings.MaterialSettings, InBaseMaterial, OutMaterial, MaterialAssetName, MaterialPackageName, OutAssetsToSync);
+	UMaterialInstanceConstant* MergedMaterial = ProxyMaterialUtilities::CreateProxyMaterialInstance(MaterialPackage, InSettings.MaterialSettings, InBaseMaterial, OutMaterial, MaterialPackageName, MaterialAssetName, OutAssetsToSync);
 	// Set material static lighting usage flag if project has static lighting enabled
 	static const auto AllowStaticLightingVar = IConsoleManager::Get().FindTConsoleVariableDataInt(TEXT("r.AllowStaticLighting"));
 	const bool bAllowStaticLighting = (!AllowStaticLightingVar || AllowStaticLightingVar->GetValueOnGameThread() != 0);
diff --git a/Engine/Source/Editor/AnimGraph/Private/AnimationGraphSchema.cpp b/Engine/Source/Editor/AnimGraph/Private/AnimationGraphSchema.cpp
--- a/Engine/Source/Editor/AnimGraph/Private/AnimationGraphSchema.cpp
+++ b/Engine/Source/Editor/AnimGraph/Private/AnimationGraphSchema.cpp
@@ -367,7 +367,9 @@ void UAnimationGraphSchema::DroppedAssetsOnPin(const TArray<FAssetData>& Assets,
 	UAnimationAsset* Asset = FAssetData::GetFirstAsset<UAnimationAsset>(Assets);
 	if ((Asset != NULL) && (Pin != NULL))
 	{
-		SpawnNodeFromAsset(Asset, GraphPosition, Pin->GetOwningNode()->GetGraph(), Pin);
+		// Don't have access to bounding information for node, using fixed offset that should work for most cases.
+		const FVector2D FixedOffset(-250.0f, 50.0f);
+		SpawnNodeFromAsset(Asset, GraphPosition + FixedOffset, Pin->GetOwningNode()->GetGraph(), Pin);
 	}
 }
 
diff --git a/Engine/Source/Editor/AnimationEditor/AnimationEditor.Build.cs b/Engine/Source/Editor/AnimationEditor/AnimationEditor.Build.cs
--- a/Engine/Source/Editor/AnimationEditor/AnimationEditor.Build.cs
+++ b/Engine/Source/Editor/AnimationEditor/AnimationEditor.Build.cs
@@ -20,19 +20,20 @@ public class AnimationEditor : ModuleRules
                 "SkeletonEditor",
                 "Kismet",
                 "AnimGraph",
-                "SequenceRecorder",
             }
 		);
 
         PrivateIncludePathModuleNames.AddRange(
             new string[] {
                 "PropertyEditor",
+                "SequenceRecorder",
             }
         );
 
         DynamicallyLoadedModuleNames.AddRange(
             new string[] {
                 "PropertyEditor",
+                "SequenceRecorder",
             }
         );
     }
diff --git a/Engine/Source/Editor/AnimationEditor/Private/AnimationEditor.cpp b/Engine/Source/Editor/AnimationEditor/Private/AnimationEditor.cpp
--- a/Engine/Source/Editor/AnimationEditor/Private/AnimationEditor.cpp
+++ b/Engine/Source/Editor/AnimationEditor/Private/AnimationEditor.cpp
@@ -46,9 +46,9 @@
 #include "Engine/CurveTable.h"
 #include "Developer/AssetTools/Public/IAssetTools.h"
 #include "Developer/AssetTools/Public/AssetToolsModule.h"
-#include "SequenceRecorderUtils.h"
 #include "ISkeletonTreeItem.h"
 #include "Algo/Transform.h"
+#include "ISequenceRecorder.h"
 
 const FName AnimationEditorAppIdentifier = FName(TEXT("AnimationEditorApp"));
 
@@ -198,11 +198,11 @@ void FAnimationEditor::BindCommands()
 		FCanExecuteAction::CreateSP(this, &FAnimationEditor::CanApplyRawAnimChanges));
 
 	ToolkitCommands->MapAction(FAnimationEditorCommands::Get().ExportToFBX_AnimData,
-		FExecuteAction::CreateSP(this, &FAnimationEditor::OnExportToFBX, EPoseSourceOption::CurrentAnimation_AnimData),
+		FExecuteAction::CreateSP(this, &FAnimationEditor::OnExportToFBX, EExportSourceOption::CurrentAnimation_AnimData),
 		FCanExecuteAction::CreateSP(this, &FAnimationEditor::HasValidAnimationSequence));
 
 	ToolkitCommands->MapAction(FAnimationEditorCommands::Get().ExportToFBX_PreviewMesh,
-		FExecuteAction::CreateSP(this, &FAnimationEditor::OnExportToFBX, EPoseSourceOption::CurrentAnimation_PreviewMesh),
+		FExecuteAction::CreateSP(this, &FAnimationEditor::OnExportToFBX, EExportSourceOption::CurrentAnimation_PreviewMesh),
 		FCanExecuteAction::CreateSP(this, &FAnimationEditor::HasValidAnimationSequence));
 
 	ToolkitCommands->MapAction(FAnimationEditorCommands::Get().AddLoopingInterpolation,
@@ -258,17 +258,6 @@ void FAnimationEditor::ExtendToolbar()
 
 			ToolbarBuilder.BeginSection("Animation");
 			{
-				// create button
-				{
-					ToolbarBuilder.AddComboButton(
-						FUIAction(),
-						FOnGetContent::CreateSP(this, &FAnimationEditor::GenerateCreateAssetMenu),
-						LOCTEXT("CreateAsset_Label", "Create Asset"),
-						LOCTEXT("CreateAsset_ToolTip", "Create Assets for this skeleton."),
-						FSlateIcon(FEditorStyle::GetStyleSetName(), "Persona.CreateAsset")
-					);
-				}
-
 				ToolbarBuilder.AddToolBarButton(FAnimationEditorCommands::Get().ReimportAnimation);
 				ToolbarBuilder.AddToolBarButton(FAnimationEditorCommands::Get().ApplyCompression, NAME_None, LOCTEXT("Toolbar_ApplyCompression", "Compression"));
 
@@ -557,16 +546,16 @@ void FAnimationEditor::OnApplyCompression()
 	}
 }
 
-void FAnimationEditor::OnExportToFBX(const EPoseSourceOption Option)
+void FAnimationEditor::OnExportToFBX(const EExportSourceOption Option)
 {
 	UAnimSequence* AnimSequenceToRecord = nullptr;
-	if (Option == EPoseSourceOption::CurrentAnimation_AnimData)
+	if (Option == EExportSourceOption::CurrentAnimation_AnimData)
 	{
 		TArray<UObject*> AssetsToExport;
 		AssetsToExport.Add(AnimationAsset);
 		ExportToFBX(AssetsToExport, false);
 	}
-	else if (Option == EPoseSourceOption::CurrentAnimation_PreviewMesh)
+	else if (Option == EExportSourceOption::CurrentAnimation_PreviewMesh)
 	{
 		TArray<TWeakObjectPtr<UObject>> Skeletons;
 		Skeletons.Add(PersonaToolkit->GetSkeleton());
@@ -642,266 +631,6 @@ TSharedRef< SWidget > FAnimationEditor::GenerateExportAssetMenu() const
 	return MenuBuilder.MakeWidget();
 }
 
-TSharedRef< SWidget > FAnimationEditor::GenerateCreateAssetMenu() const
-{
-	const bool bShouldCloseWindowAfterMenuSelection = true;
-	FMenuBuilder MenuBuilder(bShouldCloseWindowAfterMenuSelection, NULL);
-
-	// Create Animation menu
-	MenuBuilder.BeginSection("CreateAnimation", LOCTEXT("CreateAnimationMenuHeading", "Animation"));
-	{
-		// create menu
-		MenuBuilder.AddSubMenu(
-			LOCTEXT("CreateAnimationSubmenu", "Create Animation"),
-			LOCTEXT("CreateAnimationSubmenu_ToolTip", "Create Animation for this skeleton"),
-			FNewMenuDelegate::CreateSP(this, &FAnimationEditor::FillCreateAnimationMenu),
-			false,
-			FSlateIcon(FEditorStyle::GetStyleSetName(), "Persona.AssetActions.CreateAnimAsset")
-			);
-
-		MenuBuilder.AddSubMenu(
-			LOCTEXT("CreatePoseAssetSubmenu", "Create PoseAsset"),
-			LOCTEXT("CreatePoseAsssetSubmenu_ToolTip", "Create PoseAsset for this skeleton"),
-			FNewMenuDelegate::CreateSP(this, &FAnimationEditor::FillCreatePoseAssetMenu),
-			false,
-			FSlateIcon(FEditorStyle::GetStyleSetName(), "ClassIcon.PoseAsset")
-		);
-	}
-	MenuBuilder.EndSection();
-
-	TArray<TWeakObjectPtr<UObject>> Objects;
-
-	if (PersonaToolkit->GetPreviewMesh())
-	{
-		Objects.Add(PersonaToolkit->GetPreviewMesh());
-	}
-	else
-	{
-		Objects.Add(PersonaToolkit->GetSkeleton());
-	}
-
-	AnimationEditorUtils::FillCreateAssetMenu(MenuBuilder, Objects, FAnimAssetCreated::CreateSP(this, &FAnimationEditor::HandleAssetCreated), false);
-
-	return MenuBuilder.MakeWidget();
-}
-
-void FAnimationEditor::FillCreateAnimationMenu(FMenuBuilder& MenuBuilder) const
-{
-	TArray<TWeakObjectPtr<UObject>> Objects;
-
-	if (PersonaToolkit->GetPreviewMesh())
-	{
-		Objects.Add(PersonaToolkit->GetPreviewMesh());
-	}
-	else
-	{
-		Objects.Add(PersonaToolkit->GetSkeleton());
-	}
-
-	// create rig
-	MenuBuilder.BeginSection("CreateAnimationSubMenu", LOCTEXT("CreateAnimationSubMenuHeading", "Create Animation"));
-	{
-		MenuBuilder.AddMenuEntry(
-			LOCTEXT("CreateAnimation_RefPose", "Reference Pose"),
-			LOCTEXT("CreateAnimation_RefPose_Tooltip", "Create Animation from reference pose."),
-			FSlateIcon(),
-			FUIAction(
-				FExecuteAction::CreateStatic(&AnimationEditorUtils::ExecuteNewAnimAsset<UAnimSequenceFactory, UAnimSequence>, Objects, FString("_Sequence"), FAnimAssetCreated::CreateSP(this, &FAnimationEditor::CreateAnimation, EPoseSourceOption::ReferencePose), false),
-				FCanExecuteAction()
-				)
-			);
-
-		MenuBuilder.AddMenuEntry(
-			LOCTEXT("CreateAnimation_CurrentPose", "Current Pose"),
-			LOCTEXT("CreateAnimation_CurrentPose_Tooltip", "Create Animation from current pose."),
-			FSlateIcon(),
-			FUIAction(
-				FExecuteAction::CreateStatic(&AnimationEditorUtils::ExecuteNewAnimAsset<UAnimSequenceFactory, UAnimSequence>, Objects, FString("_Sequence"), FAnimAssetCreated::CreateSP(this, &FAnimationEditor::CreateAnimation, EPoseSourceOption::CurrentPose), false),
-				FCanExecuteAction()
-				)
-			);
-
-		MenuBuilder.AddSubMenu(
-			LOCTEXT("CreateAnimation_CurrenAnimationSubMenu", "Current Animation"),
-			LOCTEXT("CreateAnimation_CurrenAnimationSubMenu_ToolTip", "Create Animation from current animation"),
-			FNewMenuDelegate::CreateSP(this, &FAnimationEditor::FillCreateAnimationFromCurrentAnimationMenu),
-			false,
-			FSlateIcon(FEditorStyle::GetStyleSetName(), "Persona.AssetActions.CreateAnimAsset")
-		);
-	}
-	MenuBuilder.EndSection();
-}
-
-void FAnimationEditor::FillCreateAnimationFromCurrentAnimationMenu(FMenuBuilder& MenuBuilder) const
-{
-	{
-		TArray<TWeakObjectPtr<UObject>> Objects;
-
-		if (PersonaToolkit->GetPreviewMesh())
-		{
-			Objects.Add(PersonaToolkit->GetPreviewMesh());
-		}
-		else
-		{
-			Objects.Add(PersonaToolkit->GetSkeleton());
-		}
-
-		// create rig
-		MenuBuilder.BeginSection("CreateAnimationSubMenu", LOCTEXT("CreateAnimationFromCurrentAnimationSubmenuHeading", "Create Animation"));
-		{
-			MenuBuilder.AddMenuEntry(
-				LOCTEXT("CreateAnimation_CurrentAnimation_AnimData", "Animation Data"),
-				LOCTEXT("CreateAnimation_CurrentAnimation_AnimData_Tooltip", "Create Animation from Animation Source Data."),
-				FSlateIcon(),
-				FUIAction(
-					FExecuteAction::CreateStatic(&AnimationEditorUtils::ExecuteNewAnimAsset<UAnimSequenceFactory, UAnimSequence>, Objects, FString("_Sequence"), FAnimAssetCreated::CreateSP(this, &FAnimationEditor::CreateAnimation, EPoseSourceOption::CurrentAnimation_AnimData), false),
-					FCanExecuteAction::CreateSP(this, &FAnimationEditor::HasValidAnimationSequence)
-				)
-			);
-
-			MenuBuilder.AddMenuEntry(
-				LOCTEXT("CreateAnimation_CurrentAnimation_PreviewMesh", "Preview Mesh"),
-				LOCTEXT("CreateAnimation_CurrentAnimation_PreviewMesh_Tooltip", "Create Animation by playing on the Current Preview Mesh, including Retargeting, Post Process Graph, or anything you see on the preview mesh."),
-				FSlateIcon(),
-				FUIAction(
-					FExecuteAction::CreateStatic(&AnimationEditorUtils::ExecuteNewAnimAsset<UAnimSequenceFactory, UAnimSequence>, Objects, FString("_Sequence"), FAnimAssetCreated::CreateSP(this, &FAnimationEditor::CreateAnimation, EPoseSourceOption::CurrentAnimation_PreviewMesh), false),
-					FCanExecuteAction::CreateSP(this, &FAnimationEditor::HasValidAnimationSequence)
-				)
-			);
-		}
-		MenuBuilder.EndSection();
-	}
-
-}
-void FAnimationEditor::FillCreatePoseAssetMenu(FMenuBuilder& MenuBuilder) const
-{
-	TArray<TWeakObjectPtr<UObject>> Objects;
-
-	if (PersonaToolkit->GetPreviewMesh())
-	{
-		Objects.Add(PersonaToolkit->GetPreviewMesh());
-	}
-	else
-	{
-		Objects.Add(PersonaToolkit->GetSkeleton());
-	}
-
-	// create rig
-	MenuBuilder.BeginSection("CreatePoseAssetSubMenu", LOCTEXT("CreatePoseAssetSubMenuHeading", "Create PoseAsset"));
-	{
-		MenuBuilder.AddMenuEntry(
-			LOCTEXT("CreatePoseAsset_CurrentPose", "Current Pose"),
-			LOCTEXT("CreatePoseAsset_CurrentPose_Tooltip", "Create PoseAsset from current pose."),
-			FSlateIcon(),
-			FUIAction(
-				FExecuteAction::CreateStatic(&AnimationEditorUtils::ExecuteNewAnimAsset<UPoseAssetFactory, UPoseAsset>, Objects, FString("_PoseAsset"), FAnimAssetCreated::CreateSP(this, &FAnimationEditor::CreatePoseAsset, EPoseSourceOption::CurrentPose), false),
-				FCanExecuteAction()
-			)
-		);
-
-		MenuBuilder.AddMenuEntry(
-			LOCTEXT("CreatePoseAsset_CurrentAnimation", "Current Animation"),
-			LOCTEXT("CreatePoseAsset_CurrentAnimation_Tooltip", "Create Animation from current animation."),
-			FSlateIcon(),
-			FUIAction(
-				FExecuteAction::CreateStatic(&AnimationEditorUtils::ExecuteNewAnimAsset<UPoseAssetFactory, UPoseAsset>, Objects, FString("_PoseAsset"), FAnimAssetCreated::CreateSP(this, &FAnimationEditor::CreatePoseAsset, EPoseSourceOption::CurrentAnimation_AnimData), false),
-				FCanExecuteAction()
-			)
-		);
-	}
-	MenuBuilder.EndSection();
-
-	// create pose asset
-	MenuBuilder.BeginSection("InsertPoseSubMenuSection", LOCTEXT("InsertPoseSubMenuSubMenuHeading", "Insert Pose"));
-	{
-		MenuBuilder.AddSubMenu(
-			LOCTEXT("InsertPoseSubmenu", "Insert Pose"),
-			LOCTEXT("InsertPoseSubmenu_ToolTip", "Insert current pose to selected PoseAsset"),
-			FNewMenuDelegate::CreateSP(this, &FAnimationEditor::FillInsertPoseMenu),
-			false,
-			FSlateIcon(FEditorStyle::GetStyleSetName(), "ClassIcon.PoseAsset")
-		);
-	}
-	MenuBuilder.EndSection();
-}
-
-void FAnimationEditor::FillInsertPoseMenu(FMenuBuilder& MenuBuilder) const
-{
-	FAssetPickerConfig AssetPickerConfig;
-
-	USkeleton* Skeleton = GetPersonaToolkit()->GetSkeleton();
-
-	/** The asset picker will only show skeletons */
-	AssetPickerConfig.Filter.ClassNames.Add(*UPoseAsset::StaticClass()->GetName());
-	AssetPickerConfig.Filter.bRecursiveClasses = false;
-	AssetPickerConfig.bAllowNullSelection = false;
-	AssetPickerConfig.Filter.TagsAndValues.Add(TEXT("Skeleton"), FAssetData(Skeleton).GetExportTextName());
-
-	/** The delegate that fires when an asset was selected */
-	AssetPickerConfig.OnAssetSelected = FOnAssetSelected::CreateRaw(this, &FAnimationEditor::InsertCurrentPoseToAsset);
-
-	/** The default view mode should be a list view */
-	AssetPickerConfig.InitialAssetViewType = EAssetViewType::List;
-
-	FContentBrowserModule& ContentBrowserModule = FModuleManager::Get().LoadModuleChecked<FContentBrowserModule>(TEXT("ContentBrowser"));
-
-	MenuBuilder.AddWidget(
-		ContentBrowserModule.Get().CreateAssetPicker(AssetPickerConfig),
-		FText::GetEmpty()
-	);
-}
-
-void FAnimationEditor::InsertCurrentPoseToAsset(const FAssetData& NewPoseAssetData)
-{
-	UPoseAsset* PoseAsset = Cast<UPoseAsset>(NewPoseAssetData.GetAsset());
-	FScopedTransaction ScopedTransaction(LOCTEXT("InsertPose", "Insert Pose"));
-
-	if (PoseAsset)
-	{
-		PoseAsset->Modify();
-
-		UDebugSkelMeshComponent* PreviewMeshComponent = PersonaToolkit->GetPreviewMeshComponent();
-		if (PreviewMeshComponent)
-		{
-			FSmartName NewPoseName;
-
-			bool bSuccess = PoseAsset->AddOrUpdatePoseWithUniqueName(PreviewMeshComponent, &NewPoseName);
-
-			if (bSuccess)
-			{
-				FFormatNamedArguments Args;
-				Args.Add(TEXT("PoseAsset"), FText::FromString(PoseAsset->GetName()));
-				Args.Add(TEXT("PoseName"), FText::FromName(NewPoseName.DisplayName));
-				FNotificationInfo Info(FText::Format(LOCTEXT("InsertPoseSucceeded", "The current pose has inserted to {PoseAsset} with {PoseName}"), Args));
-				Info.ExpireDuration = 7.0f;
-				Info.bUseLargeFont = false;
-				TSharedPtr<SNotificationItem> Notification = FSlateNotificationManager::Get().AddNotification(Info);
-				if (Notification.IsValid())
-				{
-					Notification->SetCompletionState(SNotificationItem::CS_Success);
-				}
-			}
-			else
-			{
-				FFormatNamedArguments Args;
-				Args.Add(TEXT("PoseAsset"), FText::FromString(PoseAsset->GetName()));
-				FNotificationInfo Info(FText::Format(LOCTEXT("InsertPoseFailed", "Inserting pose to asset {PoseAsset} has failed"), Args));
-				Info.ExpireDuration = 7.0f;
-				Info.bUseLargeFont = false;
-				TSharedPtr<SNotificationItem> Notification = FSlateNotificationManager::Get().AddNotification(Info);
-				if (Notification.IsValid())
-				{
-					Notification->SetCompletionState(SNotificationItem::CS_Fail);
-				}
-			}
-		}
-	}
-
-	// it doesn't work well if I leave the window open. The delegate goes weired or it stop showing the popups. 
-	FSlateApplication::Get().DismissAllMenus();
-}
-
-
 void FAnimationEditor::FillCopyToSoundWaveMenu(FMenuBuilder& MenuBuilder) const
 {
 	FAssetPickerConfig AssetPickerConfig;
@@ -997,134 +726,6 @@ void FAnimationEditor::CopyCurveToSoundWave(const FAssetData& SoundWaveAssetData
 	FSlateApplication::Get().DismissAllMenus();
 }
 
-bool FAnimationEditor::CreateAnimation(const TArray<UObject*> NewAssets, const EPoseSourceOption Option)
-{
-	bool bResult = true;
-	if (NewAssets.Num() > 0)
-	{
-		USkeletalMeshComponent* MeshComponent = PersonaToolkit->GetPreviewMeshComponent();
-		UAnimSequence* Sequence = Cast<UAnimSequence>(AnimationAsset);
-
-		for (auto NewAsset : NewAssets)
-		{
-			UAnimSequence* NewAnimSequence = Cast<UAnimSequence>(NewAsset);
-			if (NewAnimSequence)
-			{
-				switch (Option)
-				{
-				case EPoseSourceOption::ReferencePose:
-					bResult &= NewAnimSequence->CreateAnimation(MeshComponent->SkeletalMesh);
-					break;
-				case EPoseSourceOption::CurrentPose:
-					bResult &= NewAnimSequence->CreateAnimation(MeshComponent);
-					break;
-				case EPoseSourceOption::CurrentAnimation_AnimData:
-					bResult &= NewAnimSequence->CreateAnimation(Sequence);
-					break;
-				case EPoseSourceOption::CurrentAnimation_PreviewMesh:
-					bResult &= RecordMeshToAnimation(MeshComponent, NewAnimSequence);
-					break;
-				default: 
-					ensure(false);
-					break;
-				}
-			}
-		}
-
-		if (bResult)
-		{
-			HandleAssetCreated(NewAssets);
-
-			// if it created based on current mesh component, 
-			if (Option == EPoseSourceOption::CurrentPose)
-			{
-				UDebugSkelMeshComponent* PreviewMeshComponent = PersonaToolkit->GetPreviewMeshComponent();
-				if (PreviewMeshComponent && PreviewMeshComponent->PreviewInstance)
-				{
-					PreviewMeshComponent->PreviewInstance->ResetModifiedBone();
-				}
-			}
-		}
-	}
-	return true;
-}
-
-bool FAnimationEditor::CreatePoseAsset(const TArray<UObject*> NewAssets, const EPoseSourceOption Option)
-{
-	bool bResult = false;
-	if (NewAssets.Num() > 0)
-	{
-		UDebugSkelMeshComponent* PreviewComponent = PersonaToolkit->GetPreviewMeshComponent();
-		UAnimSequence* Sequence = Cast<UAnimSequence>(AnimationAsset);
-
-		for (auto NewAsset : NewAssets)
-		{
-			UPoseAsset* NewPoseAsset = Cast<UPoseAsset>(NewAsset);
-			if (NewPoseAsset)
-			{
-				switch (Option)
-				{
-				case EPoseSourceOption::CurrentPose:
-					NewPoseAsset->AddOrUpdatePoseWithUniqueName(PreviewComponent);
-					bResult = true;
-					break;
-				case EPoseSourceOption::CurrentAnimation_AnimData:
-					NewPoseAsset->CreatePoseFromAnimation(Sequence);
-					bResult = true;
-					break;
-				default:
-					ensure(false);
-					bResult = false; 
-					break;
-				}
-
-			}
-		}
-
-		// if it contains error, warn them
-		if (bResult)
-		{
-			HandleAssetCreated(NewAssets);
-
-			// if it created based on current mesh component, 
-			if (Option == EPoseSourceOption::CurrentPose)
-			{
-				PreviewComponent->PreviewInstance->ResetModifiedBone();
-			}
-		}
-		else
-		{
-			FMessageDialog::Open(EAppMsgType::Ok, LOCTEXT("FailedToCreateAsset", "Failed to create asset"));
-		}
-	}
-	return true;
-}
-
-bool FAnimationEditor::HandleAssetCreated(const TArray<UObject*> NewAssets)
-{
-	if (NewAssets.Num() > 0)
-	{
-		FAssetRegistryModule::AssetCreated(NewAssets[0]);
-
-		if (UAnimationAsset* NewAnimAsset = Cast<UAnimationAsset>(NewAssets[0]))
-		{
-			OpenNewAnimationDocumentTab(NewAnimAsset);
-		}
-		else
-		{
-			// if not, we forward to asset manager to open the asset for us
-			// this is path for animation blueprint
-			FAssetToolsModule& AssetToolsModule = FModuleManager::LoadModuleChecked<FAssetToolsModule>(TEXT("AssetTools"));
-			TWeakPtr<IAssetTypeActions> AssetTypeActions = AssetToolsModule.Get().GetAssetTypeActionsForClass(NewAssets[0]->GetClass());
-			if (AssetTypeActions.IsValid())
-			{
-				AssetTypeActions.Pin()->OpenAssetEditor(NewAssets);
-			}
-		}
-	}
-	return true;
-}
-
 void FAnimationEditor::ConditionalRefreshEditor(UObject* InObject)
 {
 	bool bInterestingAsset = true;
@@ -1182,7 +783,8 @@ void FAnimationEditor::HandleAnimationSequenceBrowserCreated(const TSharedRef<IA
 
 bool FAnimationEditor::RecordMeshToAnimation(USkeletalMeshComponent* PreviewComponent, UAnimSequence* NewAsset) const
 {
-	return SequenceRecorderUtils::RecordSingleNodeInstanceToAnimation(PreviewComponent, NewAsset);
+	ISequenceRecorder& RecorderModule = FModuleManager::Get().LoadModuleChecked<ISequenceRecorder>("SequenceRecorder");
+	return RecorderModule.RecordSingleNodeInstanceToAnimation(PreviewComponent, NewAsset);
 }
 
 #undef LOCTEXT_NAMESPACE
diff --git a/Engine/Source/Editor/AnimationEditor/Private/AnimationEditor.h b/Engine/Source/Editor/AnimationEditor/Private/AnimationEditor.h
--- a/Engine/Source/Editor/AnimationEditor/Private/AnimationEditor.h
+++ b/Engine/Source/Editor/AnimationEditor/Private/AnimationEditor.h
@@ -105,10 +105,9 @@ public:
 
 private:
 
-	enum class EPoseSourceOption : uint8
+	/** Options for asset export */
+	enum class EExportSourceOption : uint8
 	{
-		ReferencePose,
-		CurrentPose,
 		CurrentAnimation_AnimData,
 		CurrentAnimation_PreviewMesh,
 		Max
@@ -130,38 +129,21 @@ private:
 
 	void OnApplyCompression();
 
-	void OnExportToFBX(const EPoseSourceOption Option);
+	void OnExportToFBX(const EExportSourceOption Option);
 	//Return true mean the asset was exported, false it was cancel or it fail
 	bool ExportToFBX(const TArray<UObject*> NewAssets, bool bRecordAnimation);
 
 	void OnAddLoopingInterpolation();
 	void OnRemoveBoneTrack();
 
-	TSharedRef< SWidget > GenerateCreateAssetMenu() const;
 	TSharedRef< SWidget > GenerateExportAssetMenu() const;
 
-	void FillCreateAnimationMenu(FMenuBuilder& MenuBuilder) const;
-
-	void FillCreateAnimationFromCurrentAnimationMenu(FMenuBuilder& MenuBuilder) const;
-
-	void FillCreatePoseAssetMenu(FMenuBuilder& MenuBuilder) const;
-
-	void FillInsertPoseMenu(FMenuBuilder& MenuBuilder) const;
-
-	void InsertCurrentPoseToAsset(const FAssetData& NewPoseAssetData);
-
 	void FillCopyToSoundWaveMenu(FMenuBuilder& MenuBuilder) const;
 
 	void FillExportAssetMenu(FMenuBuilder& MenuBuilder) const;
 
 	void CopyCurveToSoundWave(const FAssetData& SoundWaveAssetData) const;
 
-	bool CreateAnimation(const TArray<UObject*> NewAssets, const EPoseSourceOption Option);
-
-	bool CreatePoseAsset(const TArray<UObject*> NewAssets, const EPoseSourceOption Option);
-
-	bool HandleAssetCreated(const TArray<UObject*> NewAssets);
-
 	void ConditionalRefreshEditor(UObject* InObject);
 
 	void HandlePostReimport(UObject* InObject, bool bSuccess);
diff --git a/Engine/Source/Editor/AnimationModifiers/Private/AnimationBlueprintLibrary.cpp b/Engine/Source/Editor/AnimationModifiers/Private/AnimationBlueprintLibrary.cpp
--- a/Engine/Source/Editor/AnimationModifiers/Private/AnimationBlueprintLibrary.cpp
+++ b/Engine/Source/Editor/AnimationModifiers/Private/AnimationBlueprintLibrary.cpp
@@ -714,6 +714,27 @@ void UAnimationBlueprintLibrary::ReplaceAnimNotifies(UAnimSequenceBase* Animatio
 	ReplaceAnimNotifies_Helper(AnimationSequence, OldNotifyClass.Get(), NewNotifyClass.Get(), OnNotifyReplaced, FOnNotifyStateReplaced());
 }
 
+void UAnimationBlueprintLibrary::CopyAnimNotifiesFromSequence(UAnimSequence* SrcAnimSequence, UAnimSequence* DestAnimSequence)
+{
+	if (SrcAnimSequence == nullptr)
+	{
+		UE_LOG(LogAnimationBlueprintLibrary, Warning, TEXT("Invalid Source Animation Sequence for CopyAnimNotifiesFromSequence"));
+	}
+	else if (DestAnimSequence == nullptr)
+	{
+		UE_LOG(LogAnimationBlueprintLibrary, Warning, TEXT("Invalid Destination Animation Sequence for CopyAnimNotifiesFromSequence"));
+	}
+	else if (SrcAnimSequence == DestAnimSequence)
+	{
+		UE_LOG(LogAnimationBlueprintLibrary, Warning, TEXT("Source and Destination Animation Sequence are the same for CopyAnimNotifiesFromSequence"));
+	}
+	else
+	{
+		const bool bShowDialogs = false;
+		UAnimSequence::CopyNotifies(SrcAnimSequence, DestAnimSequence, bShowDialogs);
+	}
+}
+
 int32 UAnimationBlueprintLibrary::RemoveAnimationNotifyEventsByName(UAnimSequence* AnimationSequence, FName NotifyName)
 {
 	int32 NumRemovedEvents = 0;
diff --git a/Engine/Source/Editor/AnimationModifiers/Public/AnimationBlueprintLibrary.h b/Engine/Source/Editor/AnimationModifiers/Public/AnimationBlueprintLibrary.h
--- a/Engine/Source/Editor/AnimationModifiers/Public/AnimationBlueprintLibrary.h
+++ b/Engine/Source/Editor/AnimationModifiers/Public/AnimationBlueprintLibrary.h
@@ -189,6 +189,9 @@ public:
 	UFUNCTION(BlueprintCallable, Category = "AnimationBlueprintLibrary|NotifyEvents")
 	static void ReplaceAnimNotifies(UAnimSequenceBase* AnimationSequence, TSubclassOf<UAnimNotify> OldNotifyClass, TSubclassOf<UAnimNotify> NewNotifyClass, FOnNotifyReplaced OnNotifyReplaced);
 
+	/** Copies animation notifies from Src Animation Sequence to Dest. Creates anim notify tracks as necessary. Returns true on success. */
+	UFUNCTION(BlueprintCallable, Category = "AnimationBlueprintLibrary|NotifyEvents")
+	static void CopyAnimNotifiesFromSequence(UAnimSequence* SrcAnimSequence, UAnimSequence* DestAnimSequence);
 
 	// Notify Tracks
 
diff --git a/Engine/Source/Editor/DetailCustomizations/Private/DetailCustomizations.cpp b/Engine/Source/Editor/DetailCustomizations/Private/DetailCustomizations.cpp
--- a/Engine/Source/Editor/DetailCustomizations/Private/DetailCustomizations.cpp
+++ b/Engine/Source/Editor/DetailCustomizations/Private/DetailCustomizations.cpp
@@ -134,6 +134,7 @@
 #include "SkeletalMeshLODSettingsDetails.h"
 #include "MaterialExpressionLandscapeGrassCustomization.h"
 #include "TimecodeDetailsCustomization.h"
+#include "SkeletonDetails.h"
 
 IMPLEMENT_MODULE( FDetailCustomizationsModule, DetailCustomizations );
 
@@ -372,6 +373,8 @@ void FDetailCustomizationsModule::RegisterObjectCustomizations()
 	RegisterCustomClassLayout("MaterialExpressionSetMaterialAttributes", FOnGetDetailCustomizationInstance::CreateStatic(&FMaterialAttributePropertyDetails::MakeInstance));
 	RegisterCustomClassLayout("SkeletalMeshLODSettings", FOnGetDetailCustomizationInstance::CreateStatic(&FSkeletalMeshLODSettingsDetails::MakeInstance));
 
+	RegisterCustomClassLayout("Skeleton", FOnGetDetailCustomizationInstance::CreateStatic(&FSkeletonDetails::MakeInstance));
+
 	RegisterCustomClassLayout("MotionControllerComponent", FOnGetDetailCustomizationInstance::CreateStatic(&FMotionControllerDetails::MakeInstance));
 
 	RegisterCustomClassLayout("Landscape", FOnGetDetailCustomizationInstance::CreateStatic(&FLandscapeUIDetails::MakeInstance));
diff --git a/Engine/Source/Editor/DetailCustomizations/Private/SkeletalMeshReductionSettingsDetails.cpp b/Engine/Source/Editor/DetailCustomizations/Private/SkeletalMeshReductionSettingsDetails.cpp
--- a/Engine/Source/Editor/DetailCustomizations/Private/SkeletalMeshReductionSettingsDetails.cpp
+++ b/Engine/Source/Editor/DetailCustomizations/Private/SkeletalMeshReductionSettingsDetails.cpp
@@ -13,6 +13,7 @@
 #include "SkeletalMeshReductionSettings.h"
 #include "Widgets/Input/SSpinBox.h"
 #include "Widgets/Text/STextBlock.h"
+#include "SkeletalRenderPublic.h"
 
 #define LOCTEXT_NAMESPACE "SkeletalMeshReductionSettingsDetails"
 
@@ -59,10 +60,63 @@ void FSkeletalMeshReductionSettingsDetails::CustomizeChildren(TSharedRef<IProper
 
 	bool bUseThirdPartyUI = !UseNativeReductionTool();
 
+	const int32 LODIndex = [StructPropertyHandle]() -> int32
+	{
+		if (StructPropertyHandle->GetParentHandle().IsValid())
+		{
+			if (StructPropertyHandle->GetParentHandle()->GetProperty()->GetFName() == GET_MEMBER_NAME_CHECKED(FSkeletalMeshObject, LODInfo))
+			{
+				return StructPropertyHandle->GetParentHandle()->GetIndexInArray();
+			}
+		}
 
+		return INDEX_NONE;
+	}();
+
+	auto BaseLODCustomization = [LODIndex, &StructBuilder](TSharedPtr<IPropertyHandle>& BaseLODPropertyHandle)
+	{
+		// Only able to do this for LOD2 and above, so only show the property if this is the case
+		if (LODIndex > 1)
+		{
+			// Add and retrieve the default widgets
+			IDetailPropertyRow& Row = StructBuilder.AddProperty(BaseLODPropertyHandle->AsShared());
+
+			TSharedPtr<SWidget> NameWidget;
+			TSharedPtr<SWidget> ValueWidget;
+			FDetailWidgetRow DefaultWidgetRow;
+			Row.GetDefaultWidgets(NameWidget, ValueWidget, DefaultWidgetRow);
+
+			// Customize the value property to be a spinbox with Max value cap so it is always < CurrentLODIndex
+			Row.CustomWidget()
+				.NameContent()
+				[
+					NameWidget.ToSharedRef()
+				]
+				.ValueContent()
+				.MinDesiredWidth(DefaultWidgetRow.ValueWidget.MinWidth)
+				.MaxDesiredWidth(DefaultWidgetRow.ValueWidget.MaxWidth)
+				[
+					SNew(SSpinBox<int32>)
+					.Font(IDetailLayoutBuilder::GetDetailFont())
+					.MinValue(0)
+					.MaxValue(FMath::Max(LODIndex - 1, 0))
+					.Value_Lambda([BaseLODPropertyHandle]() -> int32
+					{
+						int32 Value = INDEX_NONE;
+						BaseLODPropertyHandle->GetValue(Value);
+						return Value;
+					})					
+					.OnValueChanged_Lambda([BaseLODPropertyHandle](int32 NewValue)
+					{
+						BaseLODPropertyHandle->SetValue(NewValue);
+					})
+				];
+		}
+	};
+	
+	
 	if (bUseThirdPartyUI)
 	{
-
 		StructBuilder.AddProperty(ReductionMethodPropertyHandle.ToSharedRef());
 
 		StructBuilder.AddCustomRow(LOCTEXT("PercentTriangles_Row", "Triangle Percentage"))
@@ -125,10 +179,18 @@ void FSkeletalMeshReductionSettingsDetails::CustomizeChildren(TSharedRef<IProper
 				// we don't want to add the things that we added first
 				// maybe we make array later if we have a lot. 
 				FName PropertyName = ChildHandle->GetProperty()->GetFName();
-				if (!CustomizedProperties.Contains(PropertyName) && !CustomSimplifierOnlyProperties.Contains(PropertyName))
+
+				if (PropertyName == GET_MEMBER_NAME_CHECKED(FSkeletalMeshOptimizationSettings, BaseLOD))
 				{
-					StructBuilder.AddProperty(ChildHandle.ToSharedRef());
+					BaseLODCustomization(ChildHandle);
 				}
+				else
+				{
+					if (!CustomizedProperties.Contains(PropertyName) && !CustomSimplifierOnlyProperties.Contains(PropertyName))
+					{
+						StructBuilder.AddProperty(ChildHandle.ToSharedRef());
+					}
+				}				
 			}
 		}
 	}
@@ -181,13 +243,17 @@ void FSkeletalMeshReductionSettingsDetails::CustomizeChildren(TSharedRef<IProper
 		TSharedPtr<IPropertyHandle> MaxNumOfVertsPropertyHandle = PropertyHandles.FindChecked(GET_MEMBER_NAME_CHECKED(FSkeletalMeshOptimizationSettings, MaxNumOfVerts));
 		TSharedPtr<IPropertyHandle> MaxNumOfTrisPropertyHandle = PropertyHandles.FindChecked(GET_MEMBER_NAME_CHECKED(FSkeletalMeshOptimizationSettings, MaxNumOfTriangles));
 
-		for (auto Iter(PropertyHandles.CreateConstIterator()); Iter; ++Iter)
+		for (auto Iter(PropertyHandles.CreateIterator()); Iter; ++Iter)
 		{
 			if (UnwantedPropertyHandles.Contains(Iter.Value())) 
 			{
 				IDetailPropertyRow& SettingsRow = StructBuilder.AddProperty(Iter.Value().ToSharedRef());
 				SettingsRow.Visibility(TAttribute<EVisibility>(this, &FSkeletalMeshReductionSettingsDetails::GetVisibilityForThirdPartyTool));
 			}
+			else if (Iter.Value()->GetProperty()->GetFName() == GET_MEMBER_NAME_CHECKED(FSkeletalMeshOptimizationSettings, BaseLOD))
+			{
+				BaseLODCustomization(Iter.Value());
+			}
 			else
 			{
 				IDetailPropertyRow& SettingsRow = StructBuilder.AddProperty(Iter.Value().ToSharedRef());
diff --git a/Engine/Source/Editor/DetailCustomizations/Private/SkeletonDetails.cpp b/Engine/Source/Editor/DetailCustomizations/Private/SkeletonDetails.cpp
new file mode 100644
--- /dev/null
+++ b/Engine/Source/Editor/DetailCustomizations/Private/SkeletonDetails.cpp
@@ -0,0 +1,16 @@
+// Copyright 1998-2019 Epic Games, Inc. All Rights Reserved.
+
+#include "SkeletonDetails.h"
+#include "DetailLayoutBuilder.h"
+#include "Animation/Skeleton.h"
+
+TSharedRef<IDetailCustomization> FSkeletonDetails::MakeInstance()
+{
+	return MakeShareable(new FSkeletonDetails());
+}
+
+void FSkeletonDetails::CustomizeDetails(IDetailLayoutBuilder& DetailBuilder)
+{
+	DetailBuilder.HideProperty(GET_MEMBER_NAME_CHECKED(USkeleton, BoneTree));
+}
+
diff --git a/Engine/Source/Editor/DetailCustomizations/Private/SkeletonDetails.h b/Engine/Source/Editor/DetailCustomizations/Private/SkeletonDetails.h
new file mode 100644
--- /dev/null
+++ b/Engine/Source/Editor/DetailCustomizations/Private/SkeletonDetails.h
@@ -0,0 +1,18 @@
+// Copyright 1998-2019 Epic Games, Inc. All Rights Reserved.
+
+#pragma once
+
+#include "CoreMinimal.h"
+#include "IDetailCustomization.h"
+
+class IDetailLayoutBuilder;
+
+class FSkeletonDetails : public IDetailCustomization
+{
+public:
+	/** Makes a new instance of this detail layout class for a specific detail view requesting it */
+	static TSharedRef<IDetailCustomization> MakeInstance();
+
+	/** IDetailCustomization interface */
+	virtual void CustomizeDetails(IDetailLayoutBuilder& DetailBuilder) override;
+};
diff --git a/Engine/Source/Editor/GraphEditor/Private/SGraphPin.cpp b/Engine/Source/Editor/GraphEditor/Private/SGraphPin.cpp
--- a/Engine/Source/Editor/GraphEditor/Private/SGraphPin.cpp
+++ b/Engine/Source/Editor/GraphEditor/Private/SGraphPin.cpp
@@ -657,26 +657,6 @@ void SGraphPin::OnDragEnter( const FGeometry& MyGeometry, const FDragDropEvent&
 		// Pins treat being dragged over the same as being hovered outside of drag and drop if they know how to respond to the drag action.
 		SBorder::OnMouseEnter( MyGeometry, DragDropEvent );
 	}
-	else if (Operation->IsOfType<FAssetDragDropOp>())
-	{
-		TSharedPtr<SGraphNode> NodeWidget = OwnerNodePtr.Pin();
-		if (NodeWidget.IsValid())
-		{
-			UEdGraphNode* Node = NodeWidget->GetNodeObj();
-			if(Node != NULL && Node->GetSchema() != NULL)
-			{
-				TSharedPtr<FAssetDragDropOp> AssetOp = StaticCastSharedPtr<FAssetDragDropOp>(Operation);
-				bool bOkIcon = false;
-				FString TooltipText;
-				if (AssetOp->HasAssets())
-				{
-					Node->GetSchema()->GetAssetsPinHoverMessage(AssetOp->GetAssets(), GraphPinObj, TooltipText, bOkIcon);
-				}
-				const FSlateBrush* TooltipIcon = bOkIcon ? FEditorStyle::GetBrush(TEXT("Graph.ConnectorFeedback.OK")) : FEditorStyle::GetBrush(TEXT("Graph.ConnectorFeedback.Error"));;
-				AssetOp->SetToolTip(FText::FromString(TooltipText), TooltipIcon);
-			}
-		}
-	}
 }
 
 void SGraphPin::OnDragLeave( const FDragDropEvent& DragDropEvent )
@@ -706,6 +686,35 @@ void SGraphPin::OnDragLeave( const FDragDropEvent& DragDropEvent )
 
 FReply SGraphPin::OnDragOver( const FGeometry& MyGeometry, const FDragDropEvent& DragDropEvent )
 {
+	TSharedPtr<FDragDropOperation> Operation = DragDropEvent.GetOperation();
+	if (!Operation.IsValid())
+	{
+		return FReply::Unhandled();
+	}
+
+	if (Operation->IsOfType<FAssetDragDropOp>())
+	{
+		TSharedPtr<SGraphNode> NodeWidget = OwnerNodePtr.Pin();
+		if (NodeWidget.IsValid())
+		{
+			UEdGraphNode* Node = NodeWidget->GetNodeObj();
+			if(Node != NULL && Node->GetSchema() != NULL)
+			{
+				TSharedPtr<FAssetDragDropOp> AssetOp = StaticCastSharedPtr<FAssetDragDropOp>(Operation);
+				bool bOkIcon = false;
+				FString TooltipText;
+				if (AssetOp->HasAssets())
+				{
+					Node->GetSchema()->GetAssetsPinHoverMessage(AssetOp->GetAssets(), GraphPinObj, TooltipText, bOkIcon);
+				}
+				const FSlateBrush* TooltipIcon = bOkIcon ? FEditorStyle::GetBrush(TEXT("Graph.ConnectorFeedback.OK")) : FEditorStyle::GetBrush(TEXT("Graph.ConnectorFeedback.Error"));;
+				AssetOp->SetToolTip(FText::FromString(TooltipText), TooltipIcon);
+					
+				return FReply::Handled();
+			}
+		}
+	}
+
 	return FReply::Unhandled();
 }
 
@@ -772,7 +781,7 @@ FReply SGraphPin::OnDrop( const FGeometry& MyGeometry, const FDragDropEvent& Dra
 			TSharedPtr<FAssetDragDropOp> AssetOp = StaticCastSharedPtr<FAssetDragDropOp>(Operation);
 			if (AssetOp->HasAssets())
 			{
-				Node->GetSchema()->DroppedAssetsOnPin(AssetOp->GetAssets(), DragDropEvent.GetScreenSpacePosition(), GraphPinObj);
+				Node->GetSchema()->DroppedAssetsOnPin(AssetOp->GetAssets(), NodeWidget->GetPosition() + MyGeometry.Position, GraphPinObj);
 			}
 		}
 		return FReply::Handled();
diff --git a/Engine/Source/Editor/Kismet/Private/BlueprintEditor.cpp b/Engine/Source/Editor/Kismet/Private/BlueprintEditor.cpp
--- a/Engine/Source/Editor/Kismet/Private/BlueprintEditor.cpp
+++ b/Engine/Source/Editor/Kismet/Private/BlueprintEditor.cpp
@@ -136,6 +136,8 @@
 #include "BlueprintEditorModes.h"
 #include "BlueprintEditorSettings.h"
 #include "K2Node_SwitchString.h"
+#include "AnimGraphNode_StateMachineBase.h"
+#include "AnimationStateMachineGraph.h"
 
 #include "EngineAnalytics.h"
 #include "AnalyticsEventAttribute.h"
@@ -5297,6 +5299,22 @@ void FBlueprintEditor::DeleteSelectedNodes()
 
 	SetUISelectionState(NAME_None);
 
+	// this closes all the document that is outered by this node
+	// this is used by AnimBP statemachines/states that can create subgraph
+	auto CloseAllDocumentsTab = [this](const UEdGraphNode* InNode)
+	{
+		TArray<UObject*> NodesToClose;
+		GetObjectsWithOuter(InNode, NodesToClose);
+		for (UObject* Node : NodesToClose)
+		{
+			UEdGraph* NodeGraph = Cast<UEdGraph>(Node);
+			if (NodeGraph)
+			{
+				CloseDocumentTab(NodeGraph);
+			}
+		}
+	};
+
 	for (FGraphPanelSelectionSet::TConstIterator NodeIt( SelectedNodes ); NodeIt; ++NodeIt)
 	{
 		if (UEdGraphNode* Node = Cast<UEdGraphNode>(*NodeIt))
@@ -5308,10 +5326,14 @@ void FBlueprintEditor::DeleteSelectedNodes()
 				UAnimStateNodeBase* StateNode = Cast<UAnimStateNodeBase>(Node);
 				if (StateNode)
 				{
-					UEdGraph* NodeGraph = StateNode->GetBoundGraph();
-					if (NodeGraph)
+					CloseAllDocumentsTab(StateNode);
+				}
+				else
+				{
+					const UAnimGraphNode_StateMachineBase* SMNode = Cast<UAnimGraphNode_StateMachineBase>(Node);
+					if (SMNode)
 					{
-						CloseDocumentTab(NodeGraph);
+						CloseAllDocumentsTab(SMNode);
 					}
 				}
 
diff --git a/Engine/Source/Editor/KismetWidgets/Private/SScrubControlPanel.cpp b/Engine/Source/Editor/KismetWidgets/Private/SScrubControlPanel.cpp
--- a/Engine/Source/Editor/KismetWidgets/Private/SScrubControlPanel.cpp
+++ b/Engine/Source/Editor/KismetWidgets/Private/SScrubControlPanel.cpp
@@ -53,6 +53,7 @@ void SScrubControlPanel::Construct( const SScrubControlPanel::FArguments& InArgs
 			SNew( SBorder )
 			[
 				SAssignNew(ScrubWidget, SScrubWidget)
+				.Clipping(EWidgetClipping::ClipToBounds)
 				.Value(InArgs._Value)
 				.NumOfKeys(InArgs._NumOfKeys)
 				.DisplayDrag(InArgs._DisplayDrag)
diff --git a/Engine/Source/Editor/Persona/Persona.Build.cs b/Engine/Source/Editor/Persona/Persona.Build.cs
--- a/Engine/Source/Editor/Persona/Persona.Build.cs
+++ b/Engine/Source/Editor/Persona/Persona.Build.cs
@@ -28,7 +28,8 @@ public class Persona : ModuleRules
                 "ContentBrowser",
                 "AssetTools",
                 "AnimationEditor",
-                "MeshReductionInterface"
+                "MeshReductionInterface",
+                "SequenceRecorder",
             }
         );
 
@@ -73,7 +74,8 @@ public class Persona : ModuleRules
                 "SkeletonEditor",
                 "AssetTools",
                 "AnimationEditor",
-                "MeshReductionInterface"
+                "MeshReductionInterface",
+                "SequenceRecorder",
             }
         );
     }
diff --git a/Engine/Source/Editor/Persona/Private/PersonaModule.cpp b/Engine/Source/Editor/Persona/Private/PersonaModule.cpp
--- a/Engine/Source/Editor/Persona/Private/PersonaModule.cpp
+++ b/Engine/Source/Editor/Persona/Private/PersonaModule.cpp
@@ -77,6 +77,12 @@
 #include "PersonaPreviewSceneRefPoseController.h"
 #include "AssetViewerSettings.h"
 #include "Customization/SkeletalMeshRegionCustomization.h"
+#include "AnimationEditorUtils.h"
+#include "Factories/AnimSequenceFactory.h"
+#include "Factories/PoseAssetFactory.h"
+#include "ScopedTransaction.h"
+#include "AnimPreviewInstance.h"
+#include "ISequenceRecorder.h"
 
 IMPLEMENT_MODULE( FPersonaModule, Persona );
 
@@ -953,6 +959,404 @@ void FPersonaModule::AddCommonToolbarExtensions(FToolBarBuilder& InToolbarBuilde
 			FSlateIcon("EditorStyle", "Persona.ToggleReferencePose", "Persona.ToggleReferencePose.Small")
 			);
 	}
+
+	if(InArgs.bCreateAsset)
+	{
+		InToolbarBuilder.AddComboButton(
+			FUIAction(),
+			FOnGetContent::CreateRaw(this, &FPersonaModule::GenerateCreateAssetMenu, WeakPersonaToolkit),
+			LOCTEXT("CreateAsset_Label", "Create Asset"),
+			LOCTEXT("CreateAsset_ToolTip", "Create Assets for this skeleton."),
+			FSlateIcon(FEditorStyle::GetStyleSetName(), "Persona.CreateAsset")
+		);
+	}
+}
+
+TSharedRef< SWidget > FPersonaModule::GenerateCreateAssetMenu(TWeakPtr<IPersonaToolkit> InWeakPersonaToolkit) const
+{
+	const bool bShouldCloseWindowAfterMenuSelection = true;
+	FMenuBuilder MenuBuilder(bShouldCloseWindowAfterMenuSelection, NULL);
+
+	// Create Animation menu
+	MenuBuilder.BeginSection("CreateAnimation", LOCTEXT("CreateAnimationMenuHeading", "Animation"));
+	{
+		// create menu
+		MenuBuilder.AddSubMenu(
+			LOCTEXT("CreateAnimationSubmenu", "Create Animation"),
+			LOCTEXT("CreateAnimationSubmenu_ToolTip", "Create Animation for this skeleton"),
+			FNewMenuDelegate::CreateRaw(this, &FPersonaModule::FillCreateAnimationMenu, InWeakPersonaToolkit),
+			false,
+			FSlateIcon(FEditorStyle::GetStyleSetName(), "Persona.AssetActions.CreateAnimAsset")
+			);
+
+		MenuBuilder.AddSubMenu(
+			LOCTEXT("CreatePoseAssetSubmenu", "Create PoseAsset"),
+			LOCTEXT("CreatePoseAsssetSubmenu_ToolTip", "Create PoseAsset for this skeleton"),
+			FNewMenuDelegate::CreateRaw(this, &FPersonaModule::FillCreatePoseAssetMenu, InWeakPersonaToolkit),
+			false,
+			FSlateIcon(FEditorStyle::GetStyleSetName(), "ClassIcon.PoseAsset")
+		);
+	}
+	MenuBuilder.EndSection();
+
+	TSharedRef<IPersonaToolkit> PersonaToolkit = InWeakPersonaToolkit.Pin().ToSharedRef();
+	TArray<TWeakObjectPtr<UObject>> Objects;
+	if (PersonaToolkit->GetPreviewMesh())
+	{
+		Objects.Add(PersonaToolkit->GetPreviewMesh());
+	}
+	else
+	{
+		Objects.Add(PersonaToolkit->GetSkeleton());
+	}
+
+	AnimationEditorUtils::FillCreateAssetMenu(MenuBuilder, Objects, FAnimAssetCreated::CreateRaw(this, &FPersonaModule::HandleAssetCreated), false);
+
+	return MenuBuilder.MakeWidget();
+}
+
+void FPersonaModule::FillCreateAnimationMenu(FMenuBuilder& MenuBuilder, TWeakPtr<IPersonaToolkit> InWeakPersonaToolkit) const
+{
+	TSharedRef<IPersonaToolkit> PersonaToolkit = InWeakPersonaToolkit.Pin().ToSharedRef();
+	TArray<TWeakObjectPtr<UObject>> Objects;
+	if (PersonaToolkit->GetPreviewMesh())
+	{
+		Objects.Add(PersonaToolkit->GetPreviewMesh());
+	}
+	else
+	{
+		Objects.Add(PersonaToolkit->GetSkeleton());
+	}
+
+	// create rig
+	MenuBuilder.BeginSection("CreateAnimationSubMenu", LOCTEXT("CreateAnimationSubMenuHeading", "Create Animation"));
+	{
+		MenuBuilder.AddMenuEntry(
+			LOCTEXT("CreateAnimation_RefPose", "Reference Pose"),
+			LOCTEXT("CreateAnimation_RefPose_Tooltip", "Create Animation from reference pose."),
+			FSlateIcon(),
+			FUIAction(
+				FExecuteAction::CreateStatic(&AnimationEditorUtils::ExecuteNewAnimAsset<UAnimSequenceFactory, UAnimSequence>, Objects, FString("_Sequence"), FAnimAssetCreated::CreateRaw(this, &FPersonaModule::CreateAnimation, EPoseSourceOption::ReferencePose, InWeakPersonaToolkit), false),
+				FCanExecuteAction()
+				)
+			);
+
+		MenuBuilder.AddMenuEntry(
+			LOCTEXT("CreateAnimation_CurrentPose", "Current Pose"),
+			LOCTEXT("CreateAnimation_CurrentPose_Tooltip", "Create Animation from current pose."),
+			FSlateIcon(),
+			FUIAction(
+				FExecuteAction::CreateStatic(&AnimationEditorUtils::ExecuteNewAnimAsset<UAnimSequenceFactory, UAnimSequence>, Objects, FString("_Sequence"), FAnimAssetCreated::CreateRaw(this, &FPersonaModule::CreateAnimation, EPoseSourceOption::CurrentPose, InWeakPersonaToolkit), false),
+				FCanExecuteAction()
+				)
+			);
+
+		if(UAnimSequence* AnimSequence = Cast<UAnimSequence>(PersonaToolkit->GetAnimationAsset()))
+		{
+			MenuBuilder.AddSubMenu(
+				LOCTEXT("CreateAnimation_CurrenAnimationSubMenu", "Current Animation"),
+				LOCTEXT("CreateAnimation_CurrenAnimationSubMenu_ToolTip", "Create Animation from current animation"),
+				FNewMenuDelegate::CreateRaw(this, &FPersonaModule::FillCreateAnimationFromCurrentAnimationMenu, InWeakPersonaToolkit),
+				false,
+				FSlateIcon(FEditorStyle::GetStyleSetName(), "Persona.AssetActions.CreateAnimAsset")
+			);
+		}
+	}
+	MenuBuilder.EndSection();
+}
+
+void FPersonaModule::FillCreateAnimationFromCurrentAnimationMenu(FMenuBuilder& MenuBuilder, TWeakPtr<IPersonaToolkit> InWeakPersonaToolkit) const
+{
+	TSharedRef<IPersonaToolkit> PersonaToolkit = InWeakPersonaToolkit.Pin().ToSharedRef();
+	TArray<TWeakObjectPtr<UObject>> Objects;
+
+	if (PersonaToolkit->GetPreviewMesh())
+	{
+		Objects.Add(PersonaToolkit->GetPreviewMesh());
+	}
+	else
+	{
+		Objects.Add(PersonaToolkit->GetSkeleton());
+	}
+
+	// create rig
+	MenuBuilder.BeginSection("CreateAnimationSubMenu", LOCTEXT("CreateAnimationFromCurrentAnimationSubmenuHeading", "Create Animation"));
+	{
+		MenuBuilder.AddMenuEntry(
+			LOCTEXT("CreateAnimation_CurrentAnimation_AnimData", "Animation Data"),
+			LOCTEXT("CreateAnimation_CurrentAnimation_AnimData_Tooltip", "Create Animation from Animation Source Data."),
+			FSlateIcon(),
+			FUIAction(
+				FExecuteAction::CreateStatic(&AnimationEditorUtils::ExecuteNewAnimAsset<UAnimSequenceFactory, UAnimSequence>, Objects, FString("_Sequence"), FAnimAssetCreated::CreateRaw(this, &FPersonaModule::CreateAnimation, EPoseSourceOption::CurrentAnimation_AnimData, InWeakPersonaToolkit), false)
+			)
+		);
+
+		MenuBuilder.AddMenuEntry(
+			LOCTEXT("CreateAnimation_CurrentAnimation_PreviewMesh", "Preview Mesh"),
+			LOCTEXT("CreateAnimation_CurrentAnimation_PreviewMesh_Tooltip", "Create Animation by playing on the Current Preview Mesh, including Retargeting, Post Process Graph, or anything you see on the preview mesh."),
+			FSlateIcon(),
+			FUIAction(
+				FExecuteAction::CreateStatic(&AnimationEditorUtils::ExecuteNewAnimAsset<UAnimSequenceFactory, UAnimSequence>, Objects, FString("_Sequence"), FAnimAssetCreated::CreateRaw(this, &FPersonaModule::CreateAnimation, EPoseSourceOption::CurrentAnimation_PreviewMesh, InWeakPersonaToolkit), false)
+			)
+		);
+	}
+	MenuBuilder.EndSection();
+}
+
+void FPersonaModule::FillCreatePoseAssetMenu(FMenuBuilder& MenuBuilder, TWeakPtr<IPersonaToolkit> InWeakPersonaToolkit) const
+{
+	TSharedRef<IPersonaToolkit> PersonaToolkit = InWeakPersonaToolkit.Pin().ToSharedRef();
+	TArray<TWeakObjectPtr<UObject>> Objects;
+
+	if (PersonaToolkit->GetPreviewMesh())
+	{
+		Objects.Add(PersonaToolkit->GetPreviewMesh());
+	}
+	else
+	{
+		Objects.Add(PersonaToolkit->GetSkeleton());
+	}
+
+	// create rig
+	MenuBuilder.BeginSection("CreatePoseAssetSubMenu", LOCTEXT("CreatePoseAssetSubMenuHeading", "Create PoseAsset"));
+	{
+		MenuBuilder.AddMenuEntry(
+			LOCTEXT("CreatePoseAsset_CurrentPose", "Current Pose"),
+			LOCTEXT("CreatePoseAsset_CurrentPose_Tooltip", "Create PoseAsset from current pose."),
+			FSlateIcon(),
+			FUIAction(
+				FExecuteAction::CreateStatic(&AnimationEditorUtils::ExecuteNewAnimAsset<UPoseAssetFactory, UPoseAsset>, Objects, FString("_PoseAsset"), FAnimAssetCreated::CreateRaw(this, &FPersonaModule::CreatePoseAsset, EPoseSourceOption::CurrentPose, InWeakPersonaToolkit), false),
+				FCanExecuteAction()
+			)
+		);
+
+		MenuBuilder.AddMenuEntry(
+			LOCTEXT("CreatePoseAsset_CurrentAnimation", "Current Animation"),
+			LOCTEXT("CreatePoseAsset_CurrentAnimation_Tooltip", "Create Animation from current animation."),
+			FSlateIcon(),
+			FUIAction(
+				FExecuteAction::CreateStatic(&AnimationEditorUtils::ExecuteNewAnimAsset<UPoseAssetFactory, UPoseAsset>, Objects, FString("_PoseAsset"), FAnimAssetCreated::CreateRaw(this, &FPersonaModule::CreatePoseAsset, EPoseSourceOption::CurrentAnimation_AnimData, InWeakPersonaToolkit), false),
+				FCanExecuteAction()
+			)
+		);
+	}
+	MenuBuilder.EndSection();
+
+	// create pose asset
+	MenuBuilder.BeginSection("InsertPoseSubMenuSection", LOCTEXT("InsertPoseSubMenuSubMenuHeading", "Insert Pose"));
+	{
+		MenuBuilder.AddSubMenu(
+			LOCTEXT("InsertPoseSubmenu", "Insert Pose"),
+			LOCTEXT("InsertPoseSubmenu_ToolTip", "Insert current pose to selected PoseAsset"),
+			FNewMenuDelegate::CreateRaw(this, &FPersonaModule::FillInsertPoseMenu, InWeakPersonaToolkit),
+			false,
+			FSlateIcon(FEditorStyle::GetStyleSetName(), "ClassIcon.PoseAsset")
+		);
+	}
+	MenuBuilder.EndSection();
+}
+
+void FPersonaModule::FillInsertPoseMenu(FMenuBuilder& MenuBuilder, TWeakPtr<IPersonaToolkit> InWeakPersonaToolkit) const
+{
+	FAssetPickerConfig AssetPickerConfig;
+
+	TSharedRef<IPersonaToolkit> PersonaToolkit = InWeakPersonaToolkit.Pin().ToSharedRef();
+	USkeleton* Skeleton = PersonaToolkit->GetSkeleton();
+
+	/** The asset picker will only show skeletons */
+	AssetPickerConfig.Filter.ClassNames.Add(*UPoseAsset::StaticClass()->GetName());
+	AssetPickerConfig.Filter.bRecursiveClasses = false;
+	AssetPickerConfig.bAllowNullSelection = false;
+	AssetPickerConfig.Filter.TagsAndValues.Add(TEXT("Skeleton"), FAssetData(Skeleton).GetExportTextName());
+
+	/** The delegate that fires when an asset was selected */
+	AssetPickerConfig.OnAssetSelected = FOnAssetSelected::CreateRaw(this, &FPersonaModule::InsertCurrentPoseToAsset, InWeakPersonaToolkit);
+
+	/** The default view mode should be a list view */
+	AssetPickerConfig.InitialAssetViewType = EAssetViewType::List;
+
+	FContentBrowserModule& ContentBrowserModule = FModuleManager::Get().LoadModuleChecked<FContentBrowserModule>(TEXT("ContentBrowser"));
+
+	MenuBuilder.AddWidget(
+		ContentBrowserModule.Get().CreateAssetPicker(AssetPickerConfig),
+		FText::GetEmpty()
+	);
+}
+
+void FPersonaModule::InsertCurrentPoseToAsset(const FAssetData& NewPoseAssetData, TWeakPtr<IPersonaToolkit> InWeakPersonaToolkit)
+{
+	TSharedRef<IPersonaToolkit> PersonaToolkit = InWeakPersonaToolkit.Pin().ToSharedRef();
+	UPoseAsset* PoseAsset = Cast<UPoseAsset>(NewPoseAssetData.GetAsset());
+	FScopedTransaction ScopedTransaction(LOCTEXT("InsertPose", "Insert Pose"));
+
+	if (PoseAsset)
+	{
+		PoseAsset->Modify();
+
+		UDebugSkelMeshComponent* PreviewMeshComponent = PersonaToolkit->GetPreviewMeshComponent();
+		if (PreviewMeshComponent)
+		{
+			FSmartName NewPoseName;
+
+			bool bSuccess = PoseAsset->AddOrUpdatePoseWithUniqueName(PreviewMeshComponent, &NewPoseName);
+
+			if (bSuccess)
+			{
+				FFormatNamedArguments Args;
+				Args.Add(TEXT("PoseAsset"), FText::FromString(PoseAsset->GetName()));
+				Args.Add(TEXT("PoseName"), FText::FromName(NewPoseName.DisplayName));
+				FNotificationInfo Info(FText::Format(LOCTEXT("InsertPoseSucceeded", "The current pose has inserted to {PoseAsset} with {PoseName}"), Args));
+				Info.ExpireDuration = 7.0f;
+				Info.bUseLargeFont = false;
+				TSharedPtr<SNotificationItem> Notification = FSlateNotificationManager::Get().AddNotification(Info);
+				if (Notification.IsValid())
+				{
+					Notification->SetCompletionState(SNotificationItem::CS_Success);
+				}
+			}
+			else
+			{
+				FFormatNamedArguments Args;
+				Args.Add(TEXT("PoseAsset"), FText::FromString(PoseAsset->GetName()));
+				FNotificationInfo Info(FText::Format(LOCTEXT("InsertPoseFailed", "Inserting pose to asset {PoseAsset} has failed"), Args));
+				Info.ExpireDuration = 7.0f;
+				Info.bUseLargeFont = false;
+				TSharedPtr<SNotificationItem> Notification = FSlateNotificationManager::Get().AddNotification(Info);
+				if (Notification.IsValid())
+				{
+					Notification->SetCompletionState(SNotificationItem::CS_Fail);
+				}
+			}
+		}
+	}
+
+	// it doesn't work well if I leave the window open. The delegate goes weired or it stop showing the popups. 
+	FSlateApplication::Get().DismissAllMenus();
+}
+
+bool FPersonaModule::CreateAnimation(const TArray<UObject*> NewAssets, const EPoseSourceOption Option, TWeakPtr<IPersonaToolkit> InWeakPersonaToolkit)
+{
+	bool bResult = true;
+	if (NewAssets.Num() > 0)
+	{
+		TSharedRef<IPersonaToolkit> PersonaToolkit = InWeakPersonaToolkit.Pin().ToSharedRef();
+		USkeletalMeshComponent* MeshComponent = PersonaToolkit->GetPreviewMeshComponent();
+		UAnimSequence* Sequence = Cast<UAnimSequence>(PersonaToolkit->GetAnimationAsset());
+
+		for (auto NewAsset : NewAssets)
+		{
+			UAnimSequence* NewAnimSequence = Cast<UAnimSequence>(NewAsset);
+			if (NewAnimSequence)
+			{
+				switch (Option)
+				{
+				case EPoseSourceOption::ReferencePose:
+					bResult &= NewAnimSequence->CreateAnimation(MeshComponent->SkeletalMesh);
+					break;
+				case EPoseSourceOption::CurrentPose:
+					bResult &= NewAnimSequence->CreateAnimation(MeshComponent);
+					break;
+				case EPoseSourceOption::CurrentAnimation_AnimData:
+					bResult &= NewAnimSequence->CreateAnimation(Sequence);
+					break;
+				case EPoseSourceOption::CurrentAnimation_PreviewMesh:
+				{
+					ISequenceRecorder & RecorderModule = FModuleManager::Get().LoadModuleChecked<ISequenceRecorder>("SequenceRecorder");
+					bResult &= RecorderModule.RecordSingleNodeInstanceToAnimation(MeshComponent, NewAnimSequence);
+				}
+					break;
+				default: 
+					ensure(false);
+					break;
+				}
+			}
+		}
+
+		if (bResult)
+		{
+			HandleAssetCreated(NewAssets);
+
+			// if it created based on current mesh component, 
+			if (Option == EPoseSourceOption::CurrentPose)
+			{
+				UDebugSkelMeshComponent* PreviewMeshComponent = PersonaToolkit->GetPreviewMeshComponent();
+				if (PreviewMeshComponent && PreviewMeshComponent->PreviewInstance)
+				{
+					PreviewMeshComponent->PreviewInstance->ResetModifiedBone();
+				}
+			}
+		}
+	}
+	return true;
+}
+
+bool FPersonaModule::CreatePoseAsset(const TArray<UObject*> NewAssets, const EPoseSourceOption Option, TWeakPtr<IPersonaToolkit> InWeakPersonaToolkit)
+{
+	bool bResult = false;
+	if (NewAssets.Num() > 0)
+	{
+		TSharedRef<IPersonaToolkit> PersonaToolkit = InWeakPersonaToolkit.Pin().ToSharedRef();
+		UDebugSkelMeshComponent* PreviewComponent = PersonaToolkit->GetPreviewMeshComponent();
+		UAnimSequence* Sequence = Cast<UAnimSequence>(PersonaToolkit->GetAnimationAsset());
+
+		for (auto NewAsset : NewAssets)
+		{
+			UPoseAsset* NewPoseAsset = Cast<UPoseAsset>(NewAsset);
+			if (NewPoseAsset)
+			{
+				switch (Option)
+				{
+				case EPoseSourceOption::CurrentPose:
+					NewPoseAsset->AddOrUpdatePoseWithUniqueName(PreviewComponent);
+					bResult = true;
+					break;
+				case EPoseSourceOption::CurrentAnimation_AnimData:
+					NewPoseAsset->CreatePoseFromAnimation(Sequence);
+					bResult = true;
+					break;
+				default:
+					ensure(false);
+					bResult = false; 
+					break;
+				}
+
+			}
+		}
+
+		// if it contains error, warn them
+		if (bResult)
+		{
+			HandleAssetCreated(NewAssets);
+
+			// if it created based on current mesh component, 
+			if (Option == EPoseSourceOption::CurrentPose)
+			{
+				PreviewComponent->PreviewInstance->ResetModifiedBone();
+			}
+		}
+		else
+		{
+			FMessageDialog::Open(EAppMsgType::Ok, LOCTEXT("FailedToCreateAsset", "Failed to create asset"));
+		}
+	}
+	return true;
+}
+
+bool FPersonaModule::HandleAssetCreated(const TArray<UObject*> NewAssets)
+{
+	if (NewAssets.Num() > 0)
+	{
+		FAssetRegistryModule::AssetCreated(NewAssets[0]);
+
+		// forward to asset manager to open the asset for us
+		FAssetToolsModule& AssetToolsModule = FModuleManager::LoadModuleChecked<FAssetToolsModule>(TEXT("AssetTools"));
+		TWeakPtr<IAssetTypeActions> AssetTypeActions = AssetToolsModule.Get().GetAssetTypeActionsForClass(NewAssets[0]->GetClass());
+		if (AssetTypeActions.IsValid())
+		{
+			AssetTypeActions.Pin()->OpenAssetEditor(NewAssets);
+		}
+	}
+	return true;
 }
 
 void FPersonaModule::HandleNewAnimNotifyBlueprintCreated(UBlueprint* InBlueprint)
diff --git a/Engine/Source/Editor/Persona/Public/PersonaModule.h b/Engine/Source/Editor/Persona/Public/PersonaModule.h
--- a/Engine/Source/Editor/Persona/Public/PersonaModule.h
+++ b/Engine/Source/Editor/Persona/Public/PersonaModule.h
@@ -300,6 +300,7 @@ public:
 			: bPreviewMesh(true)
 			, bPreviewAnimation(true)
 			, bReferencePose(false)
+			, bCreateAsset(true)
 		{}
 
 		/** Adds a shortcut to setup a preview mesh to override the current display */
@@ -310,6 +311,9 @@ public:
 
 		/** Adds a shortcut to set the character back to reference pose (also clears all bone modifications) */
 		bool bReferencePose;
+
+		/** Adds a combo menu to allow other anim assets to be created */
+		bool bCreateAsset;
 	};
 
 	/** Add common toobar extensions */
@@ -328,6 +332,34 @@ private:
 	/** When a new anim notify state blueprint is created, this will handle post creation work such as adding non-event default nodes */
 	void HandleNewAnimNotifyStateBlueprintCreated(UBlueprint* InBlueprint);
 
+	/** Options for asset creation */
+	enum class EPoseSourceOption : uint8
+	{
+		ReferencePose,
+		CurrentPose,
+		CurrentAnimation_AnimData,
+		CurrentAnimation_PreviewMesh,
+		Max
+	};
+
+	TSharedRef< SWidget > GenerateCreateAssetMenu(TWeakPtr<IPersonaToolkit> InWeakPersonaToolkit) const;
+
+	void FillCreateAnimationMenu(FMenuBuilder& MenuBuilder, TWeakPtr<IPersonaToolkit> InWeakPersonaToolkit) const;
+
+	void FillCreateAnimationFromCurrentAnimationMenu(FMenuBuilder& MenuBuilder, TWeakPtr<IPersonaToolkit> InWeakPersonaToolkit) const;
+
+	void FillCreatePoseAssetMenu(FMenuBuilder& MenuBuilder, TWeakPtr<IPersonaToolkit> InWeakPersonaToolkit) const;
+
+	void FillInsertPoseMenu(FMenuBuilder& MenuBuilder, TWeakPtr<IPersonaToolkit> InWeakPersonaToolkit) const;
+
+	void InsertCurrentPoseToAsset(const FAssetData& NewPoseAssetData, TWeakPtr<IPersonaToolkit> InWeakPersonaToolkit);
+
+	bool CreateAnimation(const TArray<UObject*> NewAssets, const EPoseSourceOption Option, TWeakPtr<IPersonaToolkit> InWeakPersonaToolkit);
+
+	bool CreatePoseAsset(const TArray<UObject*> NewAssets, const EPoseSourceOption Option, TWeakPtr<IPersonaToolkit> InWeakPersonaToolkit);
+	
+	bool HandleAssetCreated(const TArray<UObject*> NewAssets);
+
 private:
 	TSharedPtr<FExtensibilityManager> MenuExtensibilityManager;
 	TSharedPtr<FExtensibilityManager> ToolBarExtensibilityManager;
diff --git a/Engine/Source/Editor/PinnedCommandList/Private/SPinnedCommandList.cpp b/Engine/Source/Editor/PinnedCommandList/Private/SPinnedCommandList.cpp
--- a/Engine/Source/Editor/PinnedCommandList/Private/SPinnedCommandList.cpp
+++ b/Engine/Source/Editor/PinnedCommandList/Private/SPinnedCommandList.cpp
@@ -97,8 +97,8 @@ public:
 
 		// Using a menu builder here is slightly wasteful, but as we cant construct
 		// a menu item individually it will have to do for now.
-		const bool bInShouldCloseWindowAfterMenuSelection = true;
-		FMenuBuilder MenuBuilder(bInShouldCloseWindowAfterMenuSelection, CommandList.Pin());
+		const bool bInShouldCloseWindowAfterMenuSelection = false;
+		FMenuBuilder MenuBuilder(bInShouldCloseWindowAfterMenuSelection, CommandList.Pin(), nullptr, false, InArgs._StyleSet, false);
 		MenuBuilder.SetStyle(InArgs._StyleSet, InArgs._StyleName);
 
 		if(CommandInfo.IsValid() && CommandList.IsValid())
@@ -136,7 +136,7 @@ public:
 					]
 				];
 
-			MenuBuilder.AddWidget(CustomWidgetContainer, FText(), true);
+			MenuBuilder.AddWidget(CustomWidgetContainer, FText(), true, false);
 		}
 		else
 		{
@@ -483,7 +483,7 @@ void SPinnedCommandList::LoadSettings()
 				}
 				else if(Command.Type == EPinnedCommandListType::CustomWidget)
 				{
-					AddCustomWidget(Command.Name);
+					AddCustomWidget_Internal(Command.Name);
 				}
 			}
 		}
@@ -591,7 +591,10 @@ void SPinnedCommandList::HandleCustomWidgetInteraction(FName InCustomWidgetIdent
 
 void SPinnedCommandList::AddCustomWidget(FName InCustomWidgetIdentifier)
 {
-	AddCustomWidget_Internal(InCustomWidgetIdentifier);
+	if(FSlateApplication::Get().GetModifierKeys().AreModifersDown(EModifierKey::Shift))
+	{
+		AddCustomWidget_Internal(InCustomWidgetIdentifier);
+	}
 }
 
 TSharedPtr<SCommand> SPinnedCommandList::AddCustomWidget_Internal(FName InCustomWidgetIdentifier, const TSharedPtr<const FUICommandList_Pinnable>& InUICommandListPinnable)
diff --git a/Engine/Source/Editor/SequenceRecorder/Private/SequenceRecorderModule.cpp b/Engine/Source/Editor/SequenceRecorder/Private/SequenceRecorderModule.cpp
--- a/Engine/Source/Editor/SequenceRecorder/Private/SequenceRecorderModule.cpp
+++ b/Engine/Source/Editor/SequenceRecorder/Private/SequenceRecorderModule.cpp
@@ -44,6 +44,8 @@
 #include "WorkflowOrientedApp/WorkflowTabManager.h"
 #include "Framework/Docking/LayoutExtender.h"
 #include "MovieSceneTimeHelpers.h"
+#include "SequenceRecorderUtils.h"
+#include "Animation/AnimSequence.h"
 
 #define LOCTEXT_NAMESPACE "SequenceRecorder"
 
@@ -773,6 +775,11 @@ class FSequenceRecorderModule : public ISequenceRecorder, private FSelfRegisteri
 		}
 	}
 
+	virtual bool RecordSingleNodeInstanceToAnimation(USkeletalMeshComponent* PreviewComponent, UAnimSequence* NewAsset) override
+	{
+		return SequenceRecorderUtils::RecordSingleNodeInstanceToAnimation(PreviewComponent, NewAsset);
+	}
+
 #if WITH_EDITOR
 	static UAnimSequence* HandleCaptureSingleFrameAnimSequence(USkeletalMeshComponent* Component)
 	{
diff --git a/Engine/Source/Editor/SequenceRecorder/Public/ISequenceRecorder.h b/Engine/Source/Editor/SequenceRecorder/Public/ISequenceRecorder.h
--- a/Engine/Source/Editor/SequenceRecorder/Public/ISequenceRecorder.h
+++ b/Engine/Source/Editor/SequenceRecorder/Public/ISequenceRecorder.h
@@ -13,6 +13,7 @@ class AActor;
 class ISequenceAudioRecorder;
 class ISequenceRecorderExtender;
 class USequenceRecordingBase;
+class UAnimSequence;
 
 DECLARE_MULTICAST_DELEGATE_OneParam(FOnRecordingStarted, class UMovieSceneSequence* /*Sequence*/);
 
@@ -165,4 +166,12 @@ public:
 
 	/** Remove an extension from the SequenceRecorder */
 	virtual void RemoveSequenceRecorderExtender(TSharedPtr<ISequenceRecorderExtender> SequenceRecorderExternder) = 0;
+
+	/**
+	* Play the current single node instance on the PreviewComponent from time [0, GetLength()], and record to NewAsset
+	*
+	* @param: PreviewComponent - this component should contains SingleNodeInstance with time-line based asset, currently support AnimSequence or AnimComposite
+	* @param: NewAsset - this is the asset that should be recorded. This will reset all animation data internally
+	*/
+	virtual bool RecordSingleNodeInstanceToAnimation(USkeletalMeshComponent* PreviewComponent, UAnimSequence* NewAsset) = 0;
 };
diff --git a/Engine/Source/Editor/SkeletalMeshEditor/Private/SkeletalMeshEditor.cpp b/Engine/Source/Editor/SkeletalMeshEditor/Private/SkeletalMeshEditor.cpp
--- a/Engine/Source/Editor/SkeletalMeshEditor/Private/SkeletalMeshEditor.cpp
+++ b/Engine/Source/Editor/SkeletalMeshEditor/Private/SkeletalMeshEditor.cpp
@@ -1049,9 +1049,10 @@ void FSkeletalMeshEditor::HandleReimportAllMeshWithNewFile(int32 SourceFileIndex
 void FSkeletalMeshEditor::ToggleMeshSectionSelection()
 {
 	TSharedRef<IPersonaPreviewScene> PreviewScene = GetPersonaToolkit()->GetPreviewScene();
+	PreviewScene->DeselectAll();
 	bool bState = !PreviewScene->AllowMeshHitProxies();
 	GetMutableDefault<UPersonaOptions>()->bAllowMeshSectionSelection = bState;
-	PreviewScene->SetAllowMeshHitProxies(bState);	
+	PreviewScene->SetAllowMeshHitProxies(bState);
 }
 
 bool FSkeletalMeshEditor::IsMeshSectionSelectionChecked() const
diff --git a/Engine/Source/Editor/SkeletonEditor/Private/SkeletonEditor.cpp b/Engine/Source/Editor/SkeletonEditor/Private/SkeletonEditor.cpp
--- a/Engine/Source/Editor/SkeletonEditor/Private/SkeletonEditor.cpp
+++ b/Engine/Source/Editor/SkeletonEditor/Private/SkeletonEditor.cpp
@@ -234,6 +234,11 @@ void FSkeletonEditor::ExtendMenu()
 	AddMenuExtender(SkeletonEditorModule.GetMenuExtensibilityManager()->GetAllExtenders(GetToolkitCommands(), GetEditingObjects()));
 }
 
+UObject* FSkeletonEditor::HandleGetAsset()
+{
+	return GetEditingObject();
+}
+
 void FSkeletonEditor::HandleObjectsSelected(const TArray<UObject*>& InObjects)
 {
 	if (DetailsView.IsValid())
diff --git a/Engine/Source/Editor/SkeletonEditor/Private/SkeletonEditor.h b/Engine/Source/Editor/SkeletonEditor/Private/SkeletonEditor.h
--- a/Engine/Source/Editor/SkeletonEditor/Private/SkeletonEditor.h
+++ b/Engine/Source/Editor/SkeletonEditor/Private/SkeletonEditor.h
@@ -75,6 +75,8 @@ public:
 	/** Get the skeleton tree widget */
 	TSharedRef<class ISkeletonTree> GetSkeletonTree() const { return SkeletonTree.ToSharedRef(); }
 
+	UObject* HandleGetAsset();
+
 public:
 	void HandleObjectsSelected(const TArray<UObject*>& InObjects);
 
diff --git a/Engine/Source/Editor/SkeletonEditor/Private/SkeletonEditorMode.cpp b/Engine/Source/Editor/SkeletonEditor/Private/SkeletonEditorMode.cpp
--- a/Engine/Source/Editor/SkeletonEditor/Private/SkeletonEditorMode.cpp
+++ b/Engine/Source/Editor/SkeletonEditor/Private/SkeletonEditorMode.cpp
@@ -33,6 +33,7 @@ FSkeletonEditorMode::FSkeletonEditorMode(TSharedRef<FWorkflowCentricApplication>
 	TabFactories.RegisterFactory(PersonaModule.CreateRetargetManagerTabFactory(InHostingApp, InSkeletonTree->GetEditableSkeleton(), SkeletonEditor->GetPersonaToolkit()->GetPreviewScene(), SkeletonEditor->OnPostUndo));
 	TabFactories.RegisterFactory(PersonaModule.CreateCurveViewerTabFactory(InHostingApp, InSkeletonTree->GetEditableSkeleton(), SkeletonEditor->GetPersonaToolkit()->GetPreviewScene(), SkeletonEditor->OnPostUndo, OnObjectsSelected));
 	TabFactories.RegisterFactory(PersonaModule.CreateSkeletonSlotNamesTabFactory(InHostingApp, InSkeletonTree->GetEditableSkeleton(), SkeletonEditor->OnPostUndo, OnObjectSelected));
+	TabFactories.RegisterFactory(PersonaModule.CreateAssetDetailsTabFactory(InHostingApp, FOnGetAsset::CreateSP(&SkeletonEditor.Get(), &FSkeletonEditor::HandleGetAsset), FOnDetailsCreated()));
 
 	TabLayout = FTabManager::NewLayout("Standalone_SkeletonEditor_Layout_v1.2")
 		->AddArea
diff --git a/Engine/Source/Runtime/AnimGraphRuntime/Private/AnimNodes/AnimNode_PoseByName.cpp b/Engine/Source/Runtime/AnimGraphRuntime/Private/AnimNodes/AnimNode_PoseByName.cpp
--- a/Engine/Source/Runtime/AnimGraphRuntime/Private/AnimNodes/AnimNode_PoseByName.cpp
+++ b/Engine/Source/Runtime/AnimGraphRuntime/Private/AnimNodes/AnimNode_PoseByName.cpp
@@ -27,6 +27,13 @@ void FAnimNode_PoseByName::RebuildPoseList(const FBoneContainer& InBoneContainer
 void FAnimNode_PoseByName::UpdateAssetPlayer(const FAnimationUpdateContext& Context)
 {
 	FAnimNode_PoseHandler::UpdateAssetPlayer(Context);
+
+	// update pose extraction context if the name differs
+	if (CurrentPoseName != PoseName)
+	{
+		RebuildPoseList(Context.AnimInstanceProxy->GetRequiredBones(), PoseAsset);
+		CurrentPoseName = PoseName;
+	}
 }
 
 void FAnimNode_PoseByName::Evaluate_AnyThread(FPoseContext& Output)
diff --git a/Engine/Source/Runtime/AnimGraphRuntime/Public/AnimNodes/AnimNode_PoseByName.h b/Engine/Source/Runtime/AnimGraphRuntime/Public/AnimNodes/AnimNode_PoseByName.h
--- a/Engine/Source/Runtime/AnimGraphRuntime/Public/AnimNodes/AnimNode_PoseByName.h
+++ b/Engine/Source/Runtime/AnimGraphRuntime/Public/AnimNodes/AnimNode_PoseByName.h
@@ -21,9 +21,14 @@ public:
 	UPROPERTY(EditAnywhere, BlueprintReadWrite, Category = Settings, meta = (PinHiddenByDefault))
 	float PoseWeight;
 
+private:
+	/** Cached current pose name used for invalidation */
+	FName CurrentPoseName;
+
 public:	
 	FAnimNode_PoseByName()
 		: PoseWeight(1.f)
+		, CurrentPoseName(NAME_None)
 	{
 	}
 
diff --git a/Engine/Source/Runtime/Engine/Classes/Animation/AnimSequence.h b/Engine/Source/Runtime/Engine/Classes/Animation/AnimSequence.h
--- a/Engine/Source/Runtime/Engine/Classes/Animation/AnimSequence.h
+++ b/Engine/Source/Runtime/Engine/Classes/Animation/AnimSequence.h
@@ -939,7 +939,7 @@ public:
 	/** 
 	 * Copy AnimNotifies from one UAnimSequence to another.
 	 */
-	static bool CopyNotifies(UAnimSequence* SourceAnimSeq, UAnimSequence* DestAnimSeq);
+	static bool CopyNotifies(UAnimSequence* SourceAnimSeq, UAnimSequence* DestAnimSeq, bool bShowDialogs = true);
 
 	/**
 	 * Flip Rotation's W For NonRoot items, and compress it again if SkelMesh exists
diff --git a/Engine/Source/Runtime/Engine/Classes/Animation/Skeleton.h b/Engine/Source/Runtime/Engine/Classes/Animation/Skeleton.h
--- a/Engine/Source/Runtime/Engine/Classes/Animation/Skeleton.h
+++ b/Engine/Source/Runtime/Engine/Classes/Animation/Skeleton.h
@@ -275,6 +275,7 @@ UCLASS(hidecategories=Object, MinimalAPI)
 class USkeleton : public UObject, public IInterface_AssetUserData, public IInterface_PreviewMeshProvider
 {
 	friend class UAnimationBlueprintLibrary;
+	friend class FSkeletonDetails;
 
 	GENERATED_UCLASS_BODY()
 
diff --git a/Engine/Source/Runtime/Engine/Private/Animation/AnimSequence.cpp b/Engine/Source/Runtime/Engine/Private/Animation/AnimSequence.cpp
--- a/Engine/Source/Runtime/Engine/Private/Animation/AnimSequence.cpp
+++ b/Engine/Source/Runtime/Engine/Private/Animation/AnimSequence.cpp
@@ -3197,7 +3197,7 @@ bool UAnimSequence::CopyAnimSequenceProperties(UAnimSequence* SourceAnimSeq, UAn
 }
 
 
-bool UAnimSequence::CopyNotifies(UAnimSequence* SourceAnimSeq, UAnimSequence* DestAnimSeq)
+bool UAnimSequence::CopyNotifies(UAnimSequence* SourceAnimSeq, UAnimSequence* DestAnimSeq, bool bShowDialogs /*= true */)
 {
 #if WITH_EDITOR
 	// Abort if source == destination.
@@ -3209,13 +3209,13 @@ bool UAnimSequence::CopyNotifies(UAnimSequence* SourceAnimSeq, UAnimSequence* De
 	// If the destination sequence is shorter than the source sequence, we'll be dropping notifies that
 	// occur at later times than the dest sequence is long.  Give the user a chance to abort if we
 	// find any notifies that won't be copied over.
-	if( DestAnimSeq->SequenceLength < SourceAnimSeq->SequenceLength )
+	if(bShowDialogs && DestAnimSeq->SequenceLength < SourceAnimSeq->SequenceLength)
 	{
 		for(int32 NotifyIndex=0; NotifyIndex<SourceAnimSeq->Notifies.Num(); ++NotifyIndex)
 		{
 			// If a notify is found which occurs off the end of the destination sequence, prompt the user to continue.
 			const FAnimNotifyEvent& SrcNotifyEvent = SourceAnimSeq->Notifies[NotifyIndex];
-			if( SrcNotifyEvent.DisplayTime_DEPRECATED > DestAnimSeq->SequenceLength )
+			if( SrcNotifyEvent.GetTriggerTime() > DestAnimSeq->SequenceLength )
 			{
 				const bool bProceed = EAppReturnType::Yes == FMessageDialog::Open( EAppMsgType::YesNo, NSLOCTEXT("UnrealEd", "SomeNotifiesWillNotBeCopiedQ", "Some notifies will not be copied because the destination sequence is not long enough.  Proceed?") );
 				if( !bProceed )
@@ -3232,7 +3232,7 @@ bool UAnimSequence::CopyNotifies(UAnimSequence* SourceAnimSeq, UAnimSequence* De
 
 	// If the destination sequence contains any notifies, ask the user if they'd like
 	// to delete the existing notifies before copying over from the source sequence.
-	if( DestAnimSeq->Notifies.Num() > 0 )
+	if(bShowDialogs && DestAnimSeq->Notifies.Num() > 0)
 	{
 		const bool bDeleteExistingNotifies = EAppReturnType::Yes == FMessageDialog::Open(EAppMsgType::YesNo, FText::Format(
 			NSLOCTEXT("UnrealEd", "DestSeqAlreadyContainsNotifiesMergeQ", "The destination sequence already contains {0} notifies.  Delete these before copying?"), FText::AsNumber(DestAnimSeq->Notifies.Num())) );
@@ -3244,7 +3244,6 @@ bool UAnimSequence::CopyNotifies(UAnimSequence* SourceAnimSeq, UAnimSequence* De
 	}
 
 	// Do the copy.
-	TArray<int32> NewNotifyIndices;
 	int32 NumNotifiesThatWereNotCopied = 0;
 
 	for(int32 NotifyIndex=0; NotifyIndex<SourceAnimSeq->Notifies.Num(); ++NotifyIndex)
@@ -3252,32 +3251,29 @@ bool UAnimSequence::CopyNotifies(UAnimSequence* SourceAnimSeq, UAnimSequence* De
 		const FAnimNotifyEvent& SrcNotifyEvent = SourceAnimSeq->Notifies[NotifyIndex];
 
 		// Skip notifies which occur at times later than the destination sequence is long.
-		if( SrcNotifyEvent.DisplayTime_DEPRECATED > DestAnimSeq->SequenceLength )
+		if( SrcNotifyEvent.GetTriggerTime() > DestAnimSeq->SequenceLength )
 		{
+			++NumNotifiesThatWereNotCopied;
 			continue;
 		}
 
-		// Do a linear-search through existing notifies to determine where
-		// to insert the new notify.
-		int32 NewNotifyIndex = 0;
-		while( NewNotifyIndex < DestAnimSeq->Notifies.Num()
-			&& DestAnimSeq->Notifies[NewNotifyIndex].DisplayTime_DEPRECATED <= SrcNotifyEvent.DisplayTime_DEPRECATED )
+		// Copy notify tracks from src to dest if they are missing
+		if (SrcNotifyEvent.TrackIndex >= DestAnimSeq->AnimNotifyTracks.Num())
 		{
-			++NewNotifyIndex;
+			for (int32 TrackIndex = DestAnimSeq->AnimNotifyTracks.Num(); TrackIndex <= SrcNotifyEvent.TrackIndex; ++TrackIndex)
+			{
+				DestAnimSeq->AnimNotifyTracks.Add(FAnimNotifyTrack(SourceAnimSeq->AnimNotifyTracks[TrackIndex].TrackName, SourceAnimSeq->AnimNotifyTracks[TrackIndex].TrackColor));
+			}
 		}
 
 		// Track the location of the new notify.
-		NewNotifyIndices.Add(NewNotifyIndex);
-
-		// Create a new empty on in the array.
-		DestAnimSeq->Notifies.InsertZeroed(NewNotifyIndex);
-
-		// Copy time and comment.
-		FAnimNotifyEvent& Notify = DestAnimSeq->Notifies[NewNotifyIndex];
-		Notify.DisplayTime_DEPRECATED = SrcNotifyEvent.DisplayTime_DEPRECATED;
-		Notify.TriggerTimeOffset = GetTriggerTimeOffsetForType( DestAnimSeq->CalculateOffsetForNotify(Notify.DisplayTime_DEPRECATED));
-		Notify.NotifyName = SrcNotifyEvent.NotifyName;
-		Notify.Duration = SrcNotifyEvent.Duration;
+		int32 NewNotifyIndex = DestAnimSeq->Notifies.AddDefaulted();
+		FAnimNotifyEvent& NotifyEvent = DestAnimSeq->Notifies[NewNotifyIndex];
+		
+		// Copy properties of the NotifyEvent
+		NotifyEvent.TrackIndex = SrcNotifyEvent.TrackIndex;
+		NotifyEvent.NotifyName = SrcNotifyEvent.NotifyName;
+		NotifyEvent.Duration = SrcNotifyEvent.Duration;
 
 		// Copy the notify itself, and point the new one at it.
 		if( SrcNotifyEvent.Notify )
@@ -3286,7 +3282,7 @@ bool UAnimSequence::CopyNotifies(UAnimSequence* SourceAnimSeq, UAnimSequence* De
 		}
 		else
 		{
-			DestAnimSeq->Notifies[NewNotifyIndex].Notify = NULL;
+			DestAnimSeq->Notifies[NewNotifyIndex].Notify = nullptr;
 		}
 
 		if( SrcNotifyEvent.NotifyStateClass )
@@ -3295,18 +3291,23 @@ bool UAnimSequence::CopyNotifies(UAnimSequence* SourceAnimSeq, UAnimSequence* De
 		}
 		else
 		{
-			DestAnimSeq->Notifies[NewNotifyIndex].NotifyStateClass = NULL;
+			DestAnimSeq->Notifies[NewNotifyIndex].NotifyStateClass = nullptr;
 		}
+		
+		// Copy notify timing
+		NotifyEvent.LinkSequence(DestAnimSeq, SrcNotifyEvent.GetTriggerTime());
+		NotifyEvent.TriggerTimeOffset = GetTriggerTimeOffsetForType(DestAnimSeq->CalculateOffsetForNotify(NotifyEvent.GetTriggerTime()));
 
 		// Make sure editor knows we've changed something.
 		DestAnimSeq->MarkPackageDirty();
+		DestAnimSeq->RefreshCacheData();
 	}
 
 	// Inform the user if some notifies weren't copied.
-	if( SourceAnimSeq->Notifies.Num() > NewNotifyIndices.Num() )
+	if(bShowDialogs && NumNotifiesThatWereNotCopied > 0)
 	{
 		FMessageDialog::Open( EAppMsgType::Ok, FText::Format(
-			NSLOCTEXT("UnrealEd", "SomeNotifiesWereNotCopiedF", "Because the destination sequence was shorter, {0} notifies were not copied."), FText::AsNumber(SourceAnimSeq->Notifies.Num() - NewNotifyIndices.Num())) );
+			NSLOCTEXT("UnrealEd", "SomeNotifiesWereNotCopiedF", "Because the destination sequence was shorter, {0} notifies were not copied."), FText::AsNumber(NumNotifiesThatWereNotCopied)) );
 	}
 #endif // WITH_EDITOR
 
diff --git a/Engine/Source/Runtime/Slate/Private/Framework/MultiBox/MultiBoxBuilder.cpp b/Engine/Source/Runtime/Slate/Private/Framework/MultiBox/MultiBoxBuilder.cpp
--- a/Engine/Source/Runtime/Slate/Private/Framework/MultiBox/MultiBoxBuilder.cpp
+++ b/Engine/Source/Runtime/Slate/Private/Framework/MultiBox/MultiBoxBuilder.cpp
@@ -168,8 +168,8 @@ void FBaseMenuBuilder::AddMenuEntry( const FUIAction& UIAction, const TSharedRef
 
 TSharedRef< class SWidget > FMenuBuilder::MakeWidget( FMultiBox::FOnMakeMultiBoxBuilderOverride* InMakeMultiBoxBuilderOverride /* = nullptr */ )
 {
-	// Make menu builders searchable
-	return MultiBox->MakeWidget( true, InMakeMultiBoxBuilderOverride );
+	// Make menu builders searchable (by default)
+	return MultiBox->MakeWidget( bSearchable, InMakeMultiBoxBuilderOverride );
 }
 
 void FMenuBuilder::BeginSection( FName InExtensionHook, const TAttribute< FText >& InHeadingText )
@@ -279,12 +279,12 @@ void FMenuBuilder::AddWrapperSubMenu( const FText& InMenuLabel, const FText& InT
 	MultiBox->AddMultiBlock( NewMenuEntryBlock );
 }
 
-void FMenuBuilder::AddWidget( TSharedRef<SWidget> InWidget, const FText& Label, bool bNoIndent, bool bSearchable )
+void FMenuBuilder::AddWidget( TSharedRef<SWidget> InWidget, const FText& Label, bool bNoIndent, bool bInSearchable )
 {
 	ApplySectionBeginning();
 
 	TSharedRef< FWidgetBlock > NewWidgetBlock(new FWidgetBlock( InWidget, Label, bNoIndent ));
-	NewWidgetBlock->SetSearchable( bSearchable );
+	NewWidgetBlock->SetSearchable( bInSearchable );
 
 	MultiBox->AddMultiBlock( NewWidgetBlock );
 }
diff --git a/Engine/Source/Runtime/Slate/Public/Framework/MultiBox/MultiBoxBuilder.h b/Engine/Source/Runtime/Slate/Public/Framework/MultiBox/MultiBoxBuilder.h
--- a/Engine/Source/Runtime/Slate/Public/Framework/MultiBox/MultiBoxBuilder.h
+++ b/Engine/Source/Runtime/Slate/Public/Framework/MultiBox/MultiBoxBuilder.h
@@ -210,11 +210,15 @@ public:
 	 * @param	InCommandList	The action list that maps command infos to delegates that should be called for each command associated with a multiblock widget
 	 * @param	bInCloseSelfOnly	True if clicking on a menu entry closes itself only and its children but not the entire stack
 	 */
-	FMenuBuilder( const bool bInShouldCloseWindowAfterMenuSelection, TSharedPtr< const FUICommandList > InCommandList, TSharedPtr<FExtender> InExtender = TSharedPtr<FExtender>(), const bool bCloseSelfOnly = false, const ISlateStyle* InStyleSet = &FCoreStyle::Get() )
+	FMenuBuilder( const bool bInShouldCloseWindowAfterMenuSelection, TSharedPtr< const FUICommandList > InCommandList, TSharedPtr<FExtender> InExtender = TSharedPtr<FExtender>(), const bool bCloseSelfOnly = false, const ISlateStyle* InStyleSet = &FCoreStyle::Get(), bool bInSearchable = true )
 		: FBaseMenuBuilder( EMultiBoxType::Menu, bInShouldCloseWindowAfterMenuSelection, InCommandList, bCloseSelfOnly, InExtender, InStyleSet )
 		, bSectionNeedsToBeApplied(false)
+		, bSearchable(bInSearchable)
 	{
-		AddSearchWidget();
+		if(bSearchable)
+		{
+			AddSearchWidget();
+		}
 	}
 
 	/**
@@ -317,6 +321,9 @@ private:
 	
 	/** True if there is a pending section that needs to be applied */
 	bool bSectionNeedsToBeApplied;
+
+	/** Whether this menu is searchable */
+	bool bSearchable;
 };
 
 
