Crash reimporting an animation
#fix close all editors when reimporting alembic skeleton asset, also make it so that in case the skeleton on the persona toolkit is null we do not refresh the editor

diff --git a/Engine/Plugins/Experimental/AlembicImporter/Source/AlembicImporter/Private/AlembicImportFactory.cpp b/Engine/Plugins/Experimental/AlembicImporter/Source/AlembicImporter/Private/AlembicImportFactory.cpp
--- a/Engine/Plugins/Experimental/AlembicImporter/Source/AlembicImporter/Private/AlembicImportFactory.cpp
+++ b/Engine/Plugins/Experimental/AlembicImporter/Source/AlembicImporter/Private/AlembicImportFactory.cpp
@@ -474,7 +474,7 @@ EReimportResult::Type UAlembicImportFactory::Reimport(UObject* Obj)
 		}
 
 		// Close possible open editors using this asset	
-		GEditor->GetEditorSubsystem<UAssetEditorSubsystem>()->CloseAllEditorsForAsset(SkeletalMesh);
+		GEditor->GetEditorSubsystem<UAssetEditorSubsystem>()->CloseAllEditorsForAsset(AnimSequence);
 
 		EReimportResult::Type Result = ReimportSkeletalMesh(SkeletalMesh);
 
diff --git a/Engine/Source/Editor/AnimationEditor/Private/AnimationEditor.cpp b/Engine/Source/Editor/AnimationEditor/Private/AnimationEditor.cpp
--- a/Engine/Source/Editor/AnimationEditor/Private/AnimationEditor.cpp
+++ b/Engine/Source/Editor/AnimationEditor/Private/AnimationEditor.cpp
@@ -735,7 +735,7 @@ void FAnimationEditor::ConditionalRefreshEditor(UObject* InObject)
 {
 	bool bInterestingAsset = true;
 
-	if (InObject != GetPersonaToolkit()->GetSkeleton() && InObject != GetPersonaToolkit()->GetSkeleton()->GetPreviewMesh() && Cast<UAnimationAsset>(InObject) != AnimationAsset)
+	if (InObject != GetPersonaToolkit()->GetSkeleton() && (GetPersonaToolkit()->GetSkeleton() && InObject != GetPersonaToolkit()->GetSkeleton()->GetPreviewMesh()) && Cast<UAnimationAsset>(InObject) != AnimationAsset)
 	{
 		bInterestingAsset = false;
 	}
@@ -761,6 +761,11 @@ void FAnimationEditor::ConditionalRefreshEditor(UObject* InObject)
 		}
 	}
 
+	if (GetPersonaToolkit()->GetSkeleton() == nullptr)
+	{
+		bInterestingAsset = false;
+	}
+
 	if (bInterestingAsset)
 	{
 		GetPersonaToolkit()->GetPreviewScene()->InvalidateViews();
