Crash when closing the animation timeline document tab while editing curve
#fix ensure that the anim document tab can never be closed

diff --git a/Engine/Source/Editor/AnimationEditor/Private/AnimationEditor.cpp b/Engine/Source/Editor/AnimationEditor/Private/AnimationEditor.cpp
--- a/Engine/Source/Editor/AnimationEditor/Private/AnimationEditor.cpp
+++ b/Engine/Source/Editor/AnimationEditor/Private/AnimationEditor.cpp
@@ -442,6 +442,7 @@ TSharedPtr<SDockTab> FAnimationEditor::OpenNewAnimationDocumentTab(UAnimationAss
 				];
 
 			OpenedTab->SetLeftContent(IDocumentation::Get()->CreateAnchor(DocumentLink));
+			OpenedTab->SetCanCloseTab(SDockTab::FCanCloseTab::CreateLambda([]() -> bool { return false; }));
 
 			TabManager->InsertNewDocumentTab(AnimationEditorTabs::DocumentTab, FTabManager::ESearchPreference::RequireClosedTab, OpenedTab.ToSharedRef());
 
