Make sure the asset browser window stays active when double clicking an anim montage.
Before this fix it would change focus to the Montage Sections tab/window, which kills the workflow when trying to browse between different montages as you had to switch back to the asset browser tab all the time after double clicking.

Now the workflow is non-disruptive.

diff --git a/Engine/Source/Editor/AnimationEditor/Private/AnimationEditor.cpp b/Engine/Source/Editor/AnimationEditor/Private/AnimationEditor.cpp
--- a/Engine/Source/Editor/AnimationEditor/Private/AnimationEditor.cpp
+++ b/Engine/Source/Editor/AnimationEditor/Private/AnimationEditor.cpp
@@ -431,7 +431,8 @@ TSharedPtr<SDockTab> FAnimationEditor::OpenNewAnimationDocumentTab(UAnimationAss
 
 		TAttribute<FText> NameAttribute = TAttribute<FText>::Create(TAttribute<FText>::FGetter::CreateStatic(&Local::GetObjectName, (UObject*)InAnimAsset));
 
-		if (SharedAnimDocumentTab.IsValid())
+		const bool bIsReusedEditor = SharedAnimDocumentTab.IsValid();
+		if (bIsReusedEditor)
 		{
 			OpenedTab = SharedAnimDocumentTab.Pin();
 			OpenedTab->SetContent(TabContents);
@@ -464,10 +465,16 @@ TSharedPtr<SDockTab> FAnimationEditor::OpenNewAnimationDocumentTab(UAnimationAss
 			SharedAnimDocumentTab = OpenedTab;
 		}
 
-		// Invoke the preview tab if this is a montage
+		// Invoke the montage sections tab, and make sure the asset browser is there and in focus when we are dealing with a montage.
 		if(InAnimAsset->IsA<UAnimMontage>())
 		{
 			TabManager->TryInvokeTab(AnimationEditorTabs::AnimMontageSectionsTab);
+
+			// Only activate the asset browser tab when this is a reused Animation Editor window.
+			if (bIsReusedEditor)
+			{
+				TabManager->TryInvokeTab(AnimationEditorTabs::AssetBrowserTab);
+			}
 			OnSectionsChanged.Broadcast();
 		}
 		else
