Removing Copy to SoundWave menu option as it is a workflow both Audio and Animation teams want to move away from

diff --git a/Engine/Source/Editor/AnimationEditor/Private/AnimationEditor.cpp b/Engine/Source/Editor/AnimationEditor/Private/AnimationEditor.cpp
--- a/Engine/Source/Editor/AnimationEditor/Private/AnimationEditor.cpp
+++ b/Engine/Source/Editor/AnimationEditor/Private/AnimationEditor.cpp
@@ -304,14 +304,6 @@ void FAnimationEditor::ExtendMenu()
 
 				MenuBuilder.AddMenuEntry(FAnimationEditorCommands::Get().AddLoopingInterpolation);
 				MenuBuilder.AddMenuEntry(FAnimationEditorCommands::Get().RemoveBoneTracks);
-
-				MenuBuilder.AddSubMenu(
-					LOCTEXT("CopyCurvesToSoundWave", "Copy Curves To SoundWave"),
-					LOCTEXT("CopyCurvesToSoundWave_ToolTip", "Copy curves from this animation to the selected SoundWave"),
-					FNewMenuDelegate::CreateSP(InAnimationEditor, &FAnimationEditor::FillCopyToSoundWaveMenu),
-					false,
-					FSlateIcon(FEditorStyle::GetStyleSetName(), "ClassIcon.")
-				);
 			}
 			MenuBuilder.EndSection();
 		}
@@ -679,29 +671,6 @@ TSharedRef< SWidget > FAnimationEditor::GenerateExportAssetMenu() const
 	FillExportAssetMenu(MenuBuilder);
 	return MenuBuilder.MakeWidget();
 }
-
-void FAnimationEditor::FillCopyToSoundWaveMenu(FMenuBuilder& MenuBuilder) const
-{
-	FAssetPickerConfig AssetPickerConfig;
-	AssetPickerConfig.Filter.ClassNames.Add(*USoundWave::StaticClass()->GetName());
-	AssetPickerConfig.bAllowNullSelection = false;
-	AssetPickerConfig.OnAssetSelected = FOnAssetSelected::CreateRaw(this, &FAnimationEditor::CopyCurveToSoundWave);
-	AssetPickerConfig.InitialAssetViewType = EAssetViewType::List;
-
-	FContentBrowserModule& ContentBrowserModule = FModuleManager::Get().LoadModuleChecked<FContentBrowserModule>(TEXT("ContentBrowser"));
-
-	MenuBuilder.AddWidget(
-		SNew(SBox)
-		.WidthOverride(300.0f)
-		.HeightOverride(300.0f)
-		[
-			ContentBrowserModule.Get().CreateAssetPicker(AssetPickerConfig)
-		],
-		FText::GetEmpty()
-	);
-
-}
-
 void FAnimationEditor::FillExportAssetMenu(FMenuBuilder& MenuBuilder) const
 {
 	MenuBuilder.BeginSection("AnimationExport", LOCTEXT("ExportAssetMenuHeading", "Export"));
diff --git a/Engine/Source/Editor/AnimationEditor/Private/AnimationEditor.h b/Engine/Source/Editor/AnimationEditor/Private/AnimationEditor.h
--- a/Engine/Source/Editor/AnimationEditor/Private/AnimationEditor.h
+++ b/Engine/Source/Editor/AnimationEditor/Private/AnimationEditor.h
@@ -141,8 +141,6 @@ private:
 
 	TSharedRef< SWidget > GenerateExportAssetMenu() const;
 
-	void FillCopyToSoundWaveMenu(FMenuBuilder& MenuBuilder) const;
-
 	void FillExportAssetMenu(FMenuBuilder& MenuBuilder) const;
 
 	void CopyCurveToSoundWave(const FAssetData& SoundWaveAssetData) const;
